Integrate GitHub pull request 321 by Marenz: Don't treat encoding warnings as errors in the shader preprocessor

diff --git a/Engine/Source/Developer/ShaderPreprocessor/Private/ShaderPreprocessor.cpp b/Engine/Source/Developer/ShaderPreprocessor/Private/ShaderPreprocessor.cpp
--- a/Engine/Source/Developer/ShaderPreprocessor/Private/ShaderPreprocessor.cpp
+++ b/Engine/Source/Developer/ShaderPreprocessor/Private/ShaderPreprocessor.cpp
@@ -28,6 +28,7 @@ static bool FilterMcppError(const FString& ErrorMsg)
 {
 	const TCHAR* SubstringsToFilter[] =
 	{
+		TEXT("Unknown encoding:"),
 		TEXT("with no newline, supplemented newline"),
 		TEXT("Converted [CR+LF] to [LF]")
 	};
