UE4 - HlslParser support for preprocessor

diff --git a/Engine/Source/Developer/ShaderPreprocessor/Private/PreprocessorPrivate.h b/Engine/Source/Developer/ShaderPreprocessor/Private/PreprocessorPrivate.h
new file mode 100644
--- /dev/null
+++ b/Engine/Source/Developer/ShaderPreprocessor/Private/PreprocessorPrivate.h
@@ -0,0 +1,79 @@
+// Copyright 1998-2014 Epic Games, Inc. All Rights Reserved.
+
+#pragma once
+
+#include "ShaderPreprocessor.h"
+#include "mcpp.h"
+
+/**
+ * Filter MCPP errors.
+ * @param ErrorMsg - The error message.
+ * @returns true if the message is valid and has not been filtered out.
+ */
+inline bool FilterMcppError(const FString& ErrorMsg)
+{
+	const TCHAR* SubstringsToFilter[] =
+	{
+		TEXT("Unknown encoding:"),
+		TEXT("with no newline, supplemented newline"),
+		TEXT("Converted [CR+LF] to [LF]")
+	};
+	const int32 FilteredSubstringCount = ARRAY_COUNT(SubstringsToFilter);
+
+	for (int32 SubstringIndex = 0; SubstringIndex < FilteredSubstringCount; ++SubstringIndex)
+	{
+		if (ErrorMsg.Contains(SubstringsToFilter[SubstringIndex]))
+		{
+			return false;
+		}
+	}
+	return true;
+}
+
+/**
+ * Parses MCPP error output.
+ * @param ShaderOutput - Shader output to which to add errors.
+ * @param McppErrors - MCPP error output.
+ */
+static bool ParseMcppErrors(TArray<FShaderCompilerError>& OutErrors, const FString& McppErrors, bool bConvertFilenameToRelative)
+{
+	bool bSuccess = true;
+	if (McppErrors.Len() > 0)
+	{
+		TArray<FString> Lines;
+		McppErrors.ParseIntoArray(&Lines, TEXT("\n"), true);
+		for (int32 LineIndex = 0; LineIndex < Lines.Num(); ++LineIndex)
+		{
+			const FString& Line = Lines[LineIndex];
+			int32 SepIndex1 = Line.Find(TEXT(":"), ESearchCase::CaseSensitive, ESearchDir::FromStart, 2);
+			int32 SepIndex2 = Line.Find(TEXT(":"), ESearchCase::CaseSensitive, ESearchDir::FromStart, SepIndex1 + 1);
+			if (SepIndex1 != INDEX_NONE && SepIndex2 != INDEX_NONE)
+			{
+				FString Filename = Line.Left(SepIndex1);
+				FString LineNumStr = Line.Mid(SepIndex1 + 1, SepIndex2 - SepIndex1 - 1);
+				FString Message = Line.Mid(SepIndex2 + 1, Line.Len() - SepIndex2 - 1);
+				if (Filename.Len() && LineNumStr.Len() && LineNumStr.IsNumeric() && Message.Len())
+				{
+					while (++LineIndex < Lines.Num() && Lines[LineIndex].Len() && Lines[LineIndex].StartsWith(TEXT(" "),ESearchCase::CaseSensitive))
+					{
+						Message += FString(TEXT("\n")) + Lines[LineIndex];
+					}
+					--LineIndex;
+					Message = Message.Trim().TrimTrailing();
+
+					// Ignore the warning about files that don't end with a newline.
+					if (FilterMcppError(Message))
+					{
+						FShaderCompilerError* CompilerError = new(OutErrors) FShaderCompilerError;
+						CompilerError->ErrorFile = bConvertFilenameToRelative ? GetRelativeShaderFilename(Filename): Filename;
+						CompilerError->ErrorLineString = LineNumStr;
+						CompilerError->StrippedErrorMessage = Message;
+						bSuccess = false;
+					}
+				}
+
+			}
+		}
+	}
+	return bSuccess;
+}
diff --git a/Engine/Source/Developer/ShaderPreprocessor/Private/ShaderPreprocessor.cpp b/Engine/Source/Developer/ShaderPreprocessor/Private/ShaderPreprocessor.cpp
--- a/Engine/Source/Developer/ShaderPreprocessor/Private/ShaderPreprocessor.cpp
+++ b/Engine/Source/Developer/ShaderPreprocessor/Private/ShaderPreprocessor.cpp
@@ -1,8 +1,8 @@
 // Copyright 1998-2014 Epic Games, Inc. All Rights Reserved.
 
 #include "ShaderPreprocessor.h"
-#include "mcpp.h"
 #include "ModuleManager.h"
+#include "PreprocessorPrivate.h"
 
 IMPLEMENT_MODULE(FDefaultModuleImpl, ShaderPreprocessor);
 
@@ -19,79 +19,6 @@ static void AddMcppDefines(FString& OutOptions, const TMap<FString,FString>& Def
 	}
 }
 
-/**
- * Filter MCPP errors.
- * @param ErrorMsg - The error message.
- * @returns true if the message is valid and has not been filtered out.
- */
-static bool FilterMcppError(const FString& ErrorMsg)
-{
-	const TCHAR* SubstringsToFilter[] =
-	{
-		TEXT("Unknown encoding:"),
-		TEXT("with no newline, supplemented newline"),
-		TEXT("Converted [CR+LF] to [LF]")
-	};
-	const int32 FilteredSubstringCount = ARRAY_COUNT(SubstringsToFilter);
-
-	for (int32 SubstringIndex = 0; SubstringIndex < FilteredSubstringCount; ++SubstringIndex)
-	{
-		if (ErrorMsg.Contains(SubstringsToFilter[SubstringIndex]))
-		{
-			return false;
-		}
-	}
-	return true;
-}
-
-/**
- * Parses MCPP error output.
- * @param ShaderOutput - Shader output to which to add errors.
- * @param McppErrors - MCPP error output.
- */
-static bool ParseMcppErrors(FShaderCompilerOutput& ShaderOutput, const FString& McppErrors)
-{
-	bool bSuccess = true;
-	if (McppErrors.Len() > 0)
-	{
-		TArray<FString> Lines;
-		McppErrors.ParseIntoArray(&Lines, TEXT("\n"), true);
-		for (int32 LineIndex = 0; LineIndex < Lines.Num(); ++LineIndex)
-		{
-			const FString& Line = Lines[LineIndex];
-			int32 SepIndex1 = Line.Find(TEXT(":"), ESearchCase::CaseSensitive, ESearchDir::FromStart, 2);
-			int32 SepIndex2 = Line.Find(TEXT(":"), ESearchCase::CaseSensitive, ESearchDir::FromStart, SepIndex1 + 1);
-			if (SepIndex1 != INDEX_NONE && SepIndex2 != INDEX_NONE)
-			{
-				FString Filename = Line.Left(SepIndex1);
-				FString LineNumStr = Line.Mid(SepIndex1 + 1, SepIndex2 - SepIndex1 - 1);
-				FString Message = Line.Mid(SepIndex2 + 1, Line.Len() - SepIndex2 - 1);
-				if (Filename.Len() && LineNumStr.Len() && LineNumStr.IsNumeric() && Message.Len())
-				{
-					while (++LineIndex < Lines.Num() && Lines[LineIndex].Len() && Lines[LineIndex].StartsWith(TEXT(" "),ESearchCase::CaseSensitive))
-					{
-						Message += FString(TEXT("\n")) + Lines[LineIndex];
-					}
-					--LineIndex;
-					Message = Message.Trim().TrimTrailing();
-
-					// Ignore the warning about files that don't end with a newline.
-					if (FilterMcppError(Message))
-					{
-						FShaderCompilerError* CompilerError = new(ShaderOutput.Errors) FShaderCompilerError;
-						CompilerError->ErrorFile = GetRelativeShaderFilename(Filename);
-						CompilerError->ErrorLineString = LineNumStr;
-						CompilerError->StrippedErrorMessage = Message;
-						bSuccess = false;
-					}
-				}
-
-			}
-		}
-	}
-	return bSuccess;
-}
-
 /**
  * Helper class used to load shader source files for MCPP.
  */
@@ -223,7 +150,7 @@ bool PreprocessShader(
 	McppOutput = McppOutAnsi;
 	McppErrors = McppErrAnsi;
 
-	if (ParseMcppErrors(ShaderOutput, McppErrors))
+	if (ParseMcppErrors(ShaderOutput.Errors, McppErrors, true))
 	{
 		// exchange strings
 		FMemory::Memswap( &OutPreprocessedShader, &McppOutput, sizeof(FString) );
diff --git a/Engine/Source/Developer/ShaderPreprocessor/Public/ShaderPreprocessor.h b/Engine/Source/Developer/ShaderPreprocessor/Public/ShaderPreprocessor.h
--- a/Engine/Source/Developer/ShaderPreprocessor/Public/ShaderPreprocessor.h
+++ b/Engine/Source/Developer/ShaderPreprocessor/Public/ShaderPreprocessor.h
@@ -19,3 +19,16 @@ extern SHADERPREPROCESSOR_API bool PreprocessShader(
 	const FShaderCompilerInput& ShaderInput,
 	const FShaderCompilerDefinitions& AdditionalDefines
 	);
+
+/**
+ * Preprocess a shader file, for non-usf files.
+ * @param OutPreprocessedShader - Upon return contains the preprocessed source code.
+ * @param ShaderOutput - ShaderOutput to which errors can be added.
+ * @param ShaderInput - The shader compiler input.
+ * @returns true if the shader is preprocessed without error.
+ */
+extern SHADERPREPROCESSOR_API bool PreprocessShaderFile(
+	FString& OutPreprocessedShader,
+	TArray<FShaderCompilerError>& OutShaderErrors,
+	const FString& InShaderFile
+	);
diff --git a/Engine/Source/Programs/CrossCompilerTool/CrossCompilerTool.Build.cs b/Engine/Source/Programs/CrossCompilerTool/CrossCompilerTool.Build.cs
--- a/Engine/Source/Programs/CrossCompilerTool/CrossCompilerTool.Build.cs
+++ b/Engine/Source/Programs/CrossCompilerTool/CrossCompilerTool.Build.cs
@@ -20,6 +20,10 @@ public class CrossCompilerTool : ModuleRules
 				"ShaderCompilerCommon",
 				"MetalShaderFormat",
 				"ShaderFormatOpenGL",
+				"ShaderPreprocessor",
+
+				//@todo-rco: Remove me!
+				"ShaderCore",
 			});
 
 		AddThirdPartyPrivateStaticDependencies(Target,
diff --git a/Engine/Source/Programs/CrossCompilerTool/Private/CrossCompilerTool.cpp b/Engine/Source/Programs/CrossCompilerTool/Private/CrossCompilerTool.cpp
--- a/Engine/Source/Programs/CrossCompilerTool/Private/CrossCompilerTool.cpp
+++ b/Engine/Source/Programs/CrossCompilerTool/Private/CrossCompilerTool.cpp
@@ -8,6 +8,7 @@
 #include "GlslBackend.h"
 #include "HlslLexer.h"
 #include "HlslParser.h"
+#include "ShaderPreprocessor.h"
 
 #include "RequiredProgramMainCPPInclude.h"
 
@@ -17,6 +18,15 @@ IMPLEMENT_APPLICATION(CrossCompilerTool, "CrossCompilerTool");
 
 namespace CCT
 {
+	static bool Preprocess(const FString& InputFile, FString& Output)
+	{
+		FShaderCompilerInput CompilerInput;
+		CompilerInput.SourceFilename = InputFile;
+		FShaderCompilerOutput CompilerOutput;
+		TArray<FShaderCompilerError> Errors;
+		return PreprocessShaderFile(Output, Errors, InputFile);
+	}
+
 	static int32 Run(const FRunInfo& RunInfo)
 	{
 		ILanguageSpec* Language = nullptr;
@@ -49,21 +59,16 @@ namespace CCT
 		}
 
 		FString HLSLShaderSource;
-		if (!FFileHelper::LoadFileToString(HLSLShaderSource, *RunInfo.InputFile))
-		{
-			UE_LOG(LogCrossCompilerTool, Error, TEXT("Couldn't load Input file '%s'!"), *RunInfo.InputFile);
-			return 1;
-		}
-
 		if (RunInfo.bUseNew)
 		{
-
-			// Assume it's preprocessed
-			////Parser.Parse(TEXT("+-1 * 2 + (3 + -4) + F(5) + G(6,7) + A++ - --B"));
-			//Parser.Parse(TEXT("void X() { a += (b ? 1 : 0); }"));
-			//Parser.Parse(TEXT("void X() { return 5 * 3 + 2; }"));
 			if (RunInfo.bList)
 			{
+				if (!FFileHelper::LoadFileToString(HLSLShaderSource, *RunInfo.InputFile))
+				{
+					UE_LOG(LogCrossCompilerTool, Error, TEXT("Couldn't load Input file '%s'!"), *RunInfo.InputFile);
+					return 1;
+				}
+
 				TArray<FString> List;
 
 				if (!FFileHelper::LoadANSITextFileToStrings(*RunInfo.InputFile, &IFileManager::Get(), List))
@@ -91,6 +96,28 @@ namespace CCT
 			}
 			else
 			{
+				if (RunInfo.bRunCPP)
+				{
+					if (!Preprocess(RunInfo.InputFile, HLSLShaderSource))
+					{
+						UE_LOG(LogCrossCompilerTool, Log, TEXT("Error during preprocessor on '%s'!"), *RunInfo.InputFile);
+						return 1;
+					}
+				}
+				else
+				{
+					if (!FFileHelper::LoadFileToString(HLSLShaderSource, *RunInfo.InputFile))
+					{
+						UE_LOG(LogCrossCompilerTool, Error, TEXT("Couldn't load Input file '%s'!"), *RunInfo.InputFile);
+						return 1;
+					}
+				}
+
+				if (RunInfo.bPreprocessOnly)
+				{
+					return 0;
+				}
+
 				if (!CrossCompiler::Parser::Parse(HLSLShaderSource, *RunInfo.InputFile))
 				{
 					UE_LOG(LogCrossCompilerTool, Log, TEXT("Error compiling '%s'!"), *RunInfo.InputFile);
diff --git a/Engine/Source/Programs/CrossCompilerTool/Private/CrossCompilerTool.h b/Engine/Source/Programs/CrossCompilerTool/Private/CrossCompilerTool.h
--- a/Engine/Source/Programs/CrossCompilerTool/Private/CrossCompilerTool.h
+++ b/Engine/Source/Programs/CrossCompilerTool/Private/CrossCompilerTool.h
@@ -26,7 +26,7 @@ namespace CCT
 		bool bRunCPP;
 		bool bUseNew;
 		bool bList;
-		//bool bCPPOnly;
+		bool bPreprocessOnly;
 
 		FRunInfo();
 		bool Setup(const FString& InOptions, const TArray<FString>& InSwitches);
diff --git a/Engine/Source/Programs/CrossCompilerTool/Private/CrossCompilerUtils.cpp b/Engine/Source/Programs/CrossCompilerTool/Private/CrossCompilerUtils.cpp
--- a/Engine/Source/Programs/CrossCompilerTool/Private/CrossCompilerUtils.cpp
+++ b/Engine/Source/Programs/CrossCompilerTool/Private/CrossCompilerUtils.cpp
@@ -14,7 +14,8 @@ namespace CCT
 		BackEnd(BE_Invalid),
 		bRunCPP(true),
 		bUseNew(false),
-		bList(false)
+		bList(false),
+		bPreprocessOnly(false)
 	{
 	}
 
@@ -25,7 +26,7 @@ namespace CCT
 		UE_LOG(LogCrossCompilerTool, Display, TEXT("\t\tOptions:"));
 		UE_LOG(LogCrossCompilerTool, Display, TEXT("\t\t\t-o=file\tOutput filename"));
 		UE_LOG(LogCrossCompilerTool, Display, TEXT("\t\t\t-entry=function\tMain entry point (defaults to Main())"));
-		//UE_LOG(LogCrossCompilerTool, Display, TEXT("\t\t\t-cpp\tOnly run C preprocessor"));
+		UE_LOG(LogCrossCompilerTool, Display, TEXT("\t\t\t-cpp\tOnly run C preprocessor"));
 		UE_LOG(LogCrossCompilerTool, Display, TEXT("\t\t\t-nocpp\tDo not run C preprocessor"));
 		UE_LOG(LogCrossCompilerTool, Display, TEXT("\t\tProfiles:"));
 		UE_LOG(LogCrossCompilerTool, Display, TEXT("\t\t\t-vs\tCompile as a Vertex Shader"));
@@ -245,6 +246,10 @@ namespace CCT
 			{
 				bRunCPP = false;
 			}
+			else if (Switch.StartsWith(TEXT("cpp")))
+			{
+				bPreprocessOnly = true;
+			}
 			else if (Switch.StartsWith(TEXT("new")))
 			{
 				bUseNew = true;
