UE4 - Missing CCT file

diff --git a/Engine/Source/Developer/ShaderPreprocessor/Private/PreprocessorFile.cpp b/Engine/Source/Developer/ShaderPreprocessor/Private/PreprocessorFile.cpp
new file mode 100644
--- /dev/null
+++ b/Engine/Source/Developer/ShaderPreprocessor/Private/PreprocessorFile.cpp
@@ -0,0 +1,128 @@
+// Copyright 1998-2014 Epic Games, Inc. All Rights Reserved.
+
+#include "ShaderPreprocessor.h"
+#include "ModuleManager.h"
+#include "PreprocessorPrivate.h"
+
+/**
+ * Helper class used to load shader source files for MCPP.
+ */
+class FSimpleMcppFileLoader
+{
+public:
+	/** Initialization constructor. */
+	explicit FSimpleMcppFileLoader(const FString& InSourceFilename)
+		//: ShaderInput(InShaderInput)
+	{
+		FString ShaderDir = FPlatformProcess::ShaderDir();
+
+		InputShaderFile = ShaderDir / FPaths::GetCleanFilename(InSourceFilename);
+
+		FString InputShaderSource;
+		if (FFileHelper::LoadFileToString(InputShaderSource, *InSourceFilename))
+		{
+			InputShaderSource = FString::Printf(TEXT("#line 1 \"%s\"\n%s"), *FPaths::GetBaseFilename(InSourceFilename), *InputShaderSource);
+			CachedFileContents.Add(GetRelativeShaderFilename(InputShaderFile),StringToArray<ANSICHAR>(*InputShaderSource, InputShaderSource.Len()));
+		}
+	}
+
+	/** Returns the input shader filename to pass to MCPP. */
+	const FString& GetInputShaderFilename() const
+	{
+		return InputShaderFile;
+	}
+
+	/** Retrieves the MCPP file loader interface. */
+	file_loader GetMcppInterface()
+	{
+		file_loader Loader;
+		Loader.get_file_contents = GetFileContents;
+		Loader.user_data = (void*)this;
+		return Loader;
+	}
+
+private:
+	/** Holder for shader contents (string + size). */
+	typedef TArray<ANSICHAR> FShaderContents;
+
+	/** MCPP callback for retrieving file contents. */
+	static int GetFileContents(void* InUserData, const ANSICHAR* InFilename, const ANSICHAR** OutContents, size_t* OutContentSize)
+	{
+		auto* This = (FSimpleMcppFileLoader*)InUserData;
+		FString Filename = GetRelativeShaderFilename(ANSI_TO_TCHAR(InFilename));
+
+		FShaderContents* CachedContents = This->CachedFileContents.Find(Filename);
+		if (!CachedContents)
+		{
+			FString FileContents;
+			FFileHelper::LoadFileToString(FileContents, *Filename);
+
+			if (FileContents.Len() > 0)
+			{
+				CachedContents = &This->CachedFileContents.Add(Filename,StringToArray<ANSICHAR>(*FileContents, FileContents.Len()));
+			}
+		}
+
+		if (OutContents)
+		{
+			*OutContents = CachedContents ? CachedContents->GetData() : nullptr;
+		}
+		if (OutContentSize)
+		{
+			*OutContentSize = CachedContents ? CachedContents->Num() : 0;
+		}
+
+		return !!CachedContents;
+	}
+
+	/** Shader input data. */
+	//const FShaderCompilerInput& ShaderInput;
+	/** File contents are cached as needed. */
+	TMap<FString,FShaderContents> CachedFileContents;
+	/** The input shader filename. */
+	FString InputShaderFile;
+};
+
+
+/**
+ * Preprocess a shader.
+ * @param OutPreprocessedShader - Upon return contains the preprocessed source code.
+ * @param ShaderOutput - ShaderOutput to which errors can be added.
+ * @param ShaderInput - The shader compiler input.
+ * @param AdditionalDefines - Additional defines with which to preprocess the shader.
+ * @returns true if the shader is preprocessed without error.
+ */
+bool PreprocessShaderFile(FString& OutPreprocessedShader, TArray<FShaderCompilerError>& OutShaderErrors, const FString& InShaderFile)
+{
+	FString McppOptions;
+	FString McppOutput, McppErrors;
+	ANSICHAR* McppOutAnsi = nullptr;
+	ANSICHAR* McppErrAnsi = nullptr;
+	bool bSuccess = false;
+
+	// MCPP is not threadsafe.
+	static FCriticalSection McppCriticalSection;
+	FScopeLock McppLock(&McppCriticalSection);
+
+	FSimpleMcppFileLoader FileLoader(InShaderFile);
+
+	int32 Result = mcpp_run(
+		TCHAR_TO_ANSI(*McppOptions),
+		TCHAR_TO_ANSI(*FileLoader.GetInputShaderFilename()),
+		&McppOutAnsi,
+		&McppErrAnsi,
+		FileLoader.GetMcppInterface()
+		);
+
+	McppOutput = McppOutAnsi;
+	McppErrors = McppErrAnsi;
+
+	if (ParseMcppErrors(OutShaderErrors, McppErrors, false))
+	{
+		// exchange strings
+		FMemory::Memswap( &OutPreprocessedShader, &McppOutput, sizeof(FString) );
+		bSuccess = true;
+	}
+
+	return bSuccess;
+}
