UE4.21 - Fix buffer underflow
#5178
#rnx

diff --git a/Engine/Source/Developer/ShaderPreprocessor/Private/ShaderPreprocessor.cpp b/Engine/Source/Developer/ShaderPreprocessor/Private/ShaderPreprocessor.cpp
--- a/Engine/Source/Developer/ShaderPreprocessor/Private/ShaderPreprocessor.cpp
+++ b/Engine/Source/Developer/ShaderPreprocessor/Private/ShaderPreprocessor.cpp
@@ -37,7 +37,7 @@ public:
 		if (LoadShaderSourceFile(*InShaderInput.VirtualSourceFilePath, InputShaderSource, nullptr))
 		{
 			InputShaderSource = FString::Printf(TEXT("%s\n#line 1\n%s"), *ShaderInput.SourceFilePrefix, *InputShaderSource);
-			CachedFileContents.Add(InShaderInput.VirtualSourceFilePath, StringToArray<ANSICHAR>(*InputShaderSource, InputShaderSource.Len()));
+			CachedFileContents.Add(InShaderInput.VirtualSourceFilePath, StringToArray<ANSICHAR>(*InputShaderSource, InputShaderSource.Len() + 1));
 		}
 	}
 
@@ -90,7 +90,7 @@ private:
 				// file path in error messages.
 				FileContents = FString::Printf(TEXT("#line 1 \"%s\"\n%s"), *VirtualFilePath, *FileContents);
 
-				CachedContents = &This->CachedFileContents.Add(VirtualFilePath, StringToArray<ANSICHAR>(*FileContents, FileContents.Len()));
+				CachedContents = &This->CachedFileContents.Add(VirtualFilePath, StringToArray<ANSICHAR>(*FileContents, FileContents.Len() + 1));
 			}
 		}
 
