Dumps the #defines set from C++ at the top of shaders in the dumped file, so a permutation of a shader class may be found with grep -r 'DIM' *

diff --git a/Engine/Source/Developer/ShaderPreprocessor/Private/ShaderPreprocessor.cpp b/Engine/Source/Developer/ShaderPreprocessor/Private/ShaderPreprocessor.cpp
--- a/Engine/Source/Developer/ShaderPreprocessor/Private/ShaderPreprocessor.cpp
+++ b/Engine/Source/Developer/ShaderPreprocessor/Private/ShaderPreprocessor.cpp
@@ -181,6 +181,20 @@ FMcppAllocator GMcppAlloc;
 
 #endif
 
+static void DumpShaderDefinesAsCommentedCode(const FShaderCompilerInput& ShaderInput, FString* OutDefines)
+{
+	const TMap<FString, FString>& Definitions = ShaderInput.Environment.GetDefinitions();
+
+	TArray<FString> Keys;
+	Definitions.GetKeys(/* out */ Keys);
+	Keys.Sort();
+
+	for (const FString& Key : Keys)
+	{
+		*OutDefines += FString::Printf(TEXT("// #define %s %s\n"), *Key, *Definitions[Key]);
+	}
+}
+
 //////////////////////////////////////////////////////////////////////////
 
 /**
@@ -262,7 +276,10 @@ bool PreprocessShader(
 		return false;
 	}
 
-	OutPreprocessedShader = MoveTemp(McppOutput);
+	// List the defines used for compilation in the preprocessed shaders, especially to know witch permutation vector this shader is.
+	DumpShaderDefinesAsCommentedCode(ShaderInput, &OutPreprocessedShader);
+
+	OutPreprocessedShader += McppOutput;
 
 	return true;
 }
