C:\Perforce\DevRendering\Engine\Source\Developer\ShaderPreprocessor\Private\ShaderPreprocessor.cpp

diff --git a/Engine/Plugins/FX/Niagara/Source/NiagaraCore/Private/NiagaraCustomVersion.cpp b/Engine/Plugins/FX/Niagara/Source/NiagaraCore/Private/NiagaraCustomVersion.cpp
--- a/Engine/Plugins/FX/Niagara/Source/NiagaraCore/Private/NiagaraCustomVersion.cpp
+++ b/Engine/Plugins/FX/Niagara/Source/NiagaraCore/Private/NiagaraCustomVersion.cpp
@@ -8,4 +8,4 @@ const FGuid FNiagaraCustomVersion::GUID(0xFCF57AFA, 0x50764283, 0xB9A9E658, 0xFF
 // Register the custom version with core
 FCustomVersionRegistration GRegisterNiagaraCustomVersion(FNiagaraCustomVersion::GUID, FNiagaraCustomVersion::LatestVersion, TEXT("NiagaraVer"));
 
-const FGuid FNiagaraCustomVersion::LatestScriptCompileVersion(0x784CB4D0, 0x1529422B, 0xB6880B12, 0xFB5A1E3F);
\ No newline at end of file
+const FGuid FNiagaraCustomVersion::LatestScriptCompileVersion(0x6484E1EE, 0x4F2E4164, 0xBCFCFB53, 0xEF31F3EB);
\ No newline at end of file
diff --git a/Engine/Source/Developer/ShaderFormatVectorVM/Private/VectorVMShaderCompiler.cpp b/Engine/Source/Developer/ShaderFormatVectorVM/Private/VectorVMShaderCompiler.cpp
--- a/Engine/Source/Developer/ShaderFormatVectorVM/Private/VectorVMShaderCompiler.cpp
+++ b/Engine/Source/Developer/ShaderFormatVectorVM/Private/VectorVMShaderCompiler.cpp
@@ -62,7 +62,8 @@ bool CompileShader_VectorVM(const FShaderCompilerInput& Input, FShaderCompilerOu
 	{
 		SCOPE_CYCLE_COUNTER(STAT_VectorVM_Compiler_CompileShader_VectorVMPreprocessShader);
 
-		if (!PreprocessShader(PreprocessedShader, Output, Input, AdditionalDefines))
+		// Don't include shader definitions since it creates shader compilation errors.
+		if (!PreprocessShader(PreprocessedShader, Output, Input, AdditionalDefines, false))
 		{
 			// The preprocessing stage will add any relevant errors.
 			return false;
diff --git a/Engine/Source/Developer/ShaderPreprocessor/Private/ShaderPreprocessor.cpp b/Engine/Source/Developer/ShaderPreprocessor/Private/ShaderPreprocessor.cpp
--- a/Engine/Source/Developer/ShaderPreprocessor/Private/ShaderPreprocessor.cpp
+++ b/Engine/Source/Developer/ShaderPreprocessor/Private/ShaderPreprocessor.cpp
@@ -203,13 +203,15 @@ static void DumpShaderDefinesAsCommentedCode(const FShaderCompilerInput& ShaderI
  * @param ShaderOutput - ShaderOutput to which errors can be added.
  * @param ShaderInput - The shader compiler input.
  * @param AdditionalDefines - Additional defines with which to preprocess the shader.
+ * @param bShaderDumpDefinesAsCommentedCode - Whether to add shader definitions as comments.
  * @returns true if the shader is preprocessed without error.
  */
 bool PreprocessShader(
 	FString& OutPreprocessedShader,
 	FShaderCompilerOutput& ShaderOutput,
 	const FShaderCompilerInput& ShaderInput,
-	const FShaderCompilerDefinitions& AdditionalDefines
+	const FShaderCompilerDefinitions& AdditionalDefines,
+	bool bShaderDumpDefinesAsCommentedCode
 	)
 {
 	// Skip the cache system and directly load the file path (used for debugging)
@@ -277,7 +279,10 @@ bool PreprocessShader(
 	}
 
 	// List the defines used for compilation in the preprocessed shaders, especially to know witch permutation vector this shader is.
-	DumpShaderDefinesAsCommentedCode(ShaderInput, &OutPreprocessedShader);
+	if (bShaderDumpDefinesAsCommentedCode)
+	{
+		DumpShaderDefinesAsCommentedCode(ShaderInput, &OutPreprocessedShader);
+	}
 
 	OutPreprocessedShader += McppOutput;
 
diff --git a/Engine/Source/Developer/ShaderPreprocessor/Public/ShaderPreprocessor.h b/Engine/Source/Developer/ShaderPreprocessor/Public/ShaderPreprocessor.h
--- a/Engine/Source/Developer/ShaderPreprocessor/Public/ShaderPreprocessor.h
+++ b/Engine/Source/Developer/ShaderPreprocessor/Public/ShaderPreprocessor.h
@@ -11,11 +11,12 @@
  * @param ShaderOutput - ShaderOutput to which errors can be added.
  * @param ShaderInput - The shader compiler input.
  * @param AdditionalDefines - Additional defines with which to preprocess the shader.
+ * @param bShaderDumpDefinesAsCommentedCode - Whether to add shader definitions as comments.
  * @returns true if the shader is preprocessed without error.
  */
 extern SHADERPREPROCESSOR_API bool PreprocessShader(
 	FString& OutPreprocessedShader,
 	FShaderCompilerOutput& ShaderOutput,
 	const FShaderCompilerInput& ShaderInput,
-	const FShaderCompilerDefinitions& AdditionalDefines
-	);
+	const FShaderCompilerDefinitions& AdditionalDefines,
+	bool bShaderDumpDefinesAsCommentedCode = true);
