Fixed crash in shader preprocessor when parsing MCPP error messages.

#rnx

diff --git a/Engine/Source/Developer/ShaderPreprocessor/Private/PreprocessorPrivate.h b/Engine/Source/Developer/ShaderPreprocessor/Private/PreprocessorPrivate.h
--- a/Engine/Source/Developer/ShaderPreprocessor/Private/PreprocessorPrivate.h
+++ b/Engine/Source/Developer/ShaderPreprocessor/Private/PreprocessorPrivate.h
@@ -68,7 +68,7 @@ static bool ParseMcppErrors(TArray<FShaderCompilerError>& OutErrors, TArray<FStr
 			const FString& Line = Lines[LineIndex];
 			int32 SepIndex1 = Line.Find(TEXT(":"), ESearchCase::CaseSensitive, ESearchDir::FromStart, 2);
 			int32 SepIndex2 = Line.Find(TEXT(":"), ESearchCase::CaseSensitive, ESearchDir::FromStart, SepIndex1 + 1);
-			if (SepIndex1 != INDEX_NONE && SepIndex2 != INDEX_NONE)
+			if (SepIndex1 != INDEX_NONE && SepIndex2 != INDEX_NONE && SepIndex1 < SepIndex2)
 			{
 				FString Filename = Line.Left(SepIndex1);
 				FString LineNumStr = Line.Mid(SepIndex1 + 1, SepIndex2 - SepIndex1 - 1);
@@ -110,7 +110,20 @@ static bool ParseMcppErrors(TArray<FShaderCompilerError>& OutErrors, TArray<FStr
 							break;
 					}
 				}
-
+				else
+				{
+					// Presume message is an error
+					FShaderCompilerError* CompilerError = new(OutErrors) FShaderCompilerError;
+					CompilerError->StrippedErrorMessage = Line;
+					bSuccess = false;
+				}
+			}
+			else
+			{
+				// Presume message is an error
+				FShaderCompilerError* CompilerError = new(OutErrors) FShaderCompilerError;
+				CompilerError->StrippedErrorMessage = Line;
+				bSuccess = false;
 			}
 		}
 	}
diff --git a/Engine/Source/Developer/ShaderPreprocessor/Private/ShaderPreprocessor.cpp b/Engine/Source/Developer/ShaderPreprocessor/Private/ShaderPreprocessor.cpp
--- a/Engine/Source/Developer/ShaderPreprocessor/Private/ShaderPreprocessor.cpp
+++ b/Engine/Source/Developer/ShaderPreprocessor/Private/ShaderPreprocessor.cpp
@@ -227,6 +227,7 @@ bool PreprocessShader(
 		check(CheckVirtualShaderFilePath(ShaderInput.VirtualSourceFilePath));
 	}
 
+	int32 Result = 0;
 	FString McppOutput, McppErrors;
 
 	static FCriticalSection McppCriticalSection;
@@ -263,7 +264,7 @@ bool PreprocessShader(
 		ANSICHAR* McppOutAnsi = NULL;
 		ANSICHAR* McppErrAnsi = NULL;
 
-		int32 Result = mcpp_run(
+		Result = mcpp_run(
 			McppOptionsANSI.GetData(),
 			McppOptionsANSI.Num(),
 			TCHAR_TO_ANSI(*ShaderInput.VirtualSourceFilePath),
@@ -272,8 +273,6 @@ bool PreprocessShader(
 			FileLoader.GetMcppInterface()
 		);
 
-		checkf(Result==0, TEXT("PreprocessShader - mcpp_run failed with error code %d"), Result);
-
 		McppOutput = McppOutAnsi;
 		McppErrors = McppErrAnsi;
 	}
@@ -283,6 +282,8 @@ bool PreprocessShader(
 		return false;
 	}
 
+	checkf(Result==0, TEXT("PreprocessShader - mcpp_run failed with error code %d"), Result);
+
 	// List the defines used for compilation in the preprocessed shaders, especially to know witch permutation vector this shader is.
 	if (bShaderDumpDefinesAsCommentedCode)
 	{
