Report unhandled mcpp failures back to UE4 instead of asserting in SWC

diff --git a/Engine/Source/Developer/ShaderPreprocessor/Private/ShaderPreprocessor.cpp b/Engine/Source/Developer/ShaderPreprocessor/Private/ShaderPreprocessor.cpp
--- a/Engine/Source/Developer/ShaderPreprocessor/Private/ShaderPreprocessor.cpp
+++ b/Engine/Source/Developer/ShaderPreprocessor/Private/ShaderPreprocessor.cpp
@@ -227,7 +227,7 @@ bool PreprocessShader(
 		check(CheckVirtualShaderFilePath(ShaderInput.VirtualSourceFilePath));
 	}
 
-	int32 Result = 0;
+	int32 McppResult = 0;
 	FString McppOutput, McppErrors;
 
 	static FCriticalSection McppCriticalSection;
@@ -264,7 +264,7 @@ bool PreprocessShader(
 		ANSICHAR* McppOutAnsi = NULL;
 		ANSICHAR* McppErrAnsi = NULL;
 
-		Result = mcpp_run(
+		McppResult = mcpp_run(
 			McppOptionsANSI.GetData(),
 			McppOptionsANSI.Num(),
 			TCHAR_TO_ANSI(*ShaderInput.VirtualSourceFilePath),
@@ -282,7 +282,15 @@ bool PreprocessShader(
 		return false;
 	}
 
-	checkf(Result==0, TEXT("PreprocessShader - mcpp_run failed with error code %d"), Result);
+	// Report unhandled mcpp failure that didn't generate any errors
+	if (McppResult != 0)
+	{
+		FShaderCompilerError* CompilerError = new(ShaderOutput.Errors) FShaderCompilerError;
+		CompilerError->ErrorVirtualFilePath = ShaderInput.VirtualSourceFilePath;
+		CompilerError->ErrorLineString = TEXT("0");
+		CompilerError->StrippedErrorMessage = FString::Printf(TEXT("PreprocessShader mcpp_run failed with error code %d"), McppResult);
+		return false;
+	}
 
 	// List the defines used for compilation in the preprocessed shaders, especially to know witch permutation vector this shader is.
 	if (bShaderDumpDefinesAsCommentedCode)
