Don't include commented out defines in the preprocessed code by default.

- Defines will still be there if we're dumping the shader code for debug.
- This makes the code less unique in general case and also avoids unnecessary works that no one sees.

diff --git a/Engine/Source/Developer/ShaderFormatVectorVM/Private/VectorVMShaderCompiler.cpp b/Engine/Source/Developer/ShaderFormatVectorVM/Private/VectorVMShaderCompiler.cpp
--- a/Engine/Source/Developer/ShaderFormatVectorVM/Private/VectorVMShaderCompiler.cpp
+++ b/Engine/Source/Developer/ShaderFormatVectorVM/Private/VectorVMShaderCompiler.cpp
@@ -79,7 +79,7 @@ bool CompileShader_VectorVM(const FShaderCompilerInput& Input, FShaderCompilerOu
 		SCOPE_CYCLE_COUNTER(STAT_VectorVM_Compiler_CompileShader_VectorVMPreprocessShader);
 
 		// Don't include shader definitions since it creates shader compilation errors.
-		if (!PreprocessShader(PreprocessedShader, Output, Input, AdditionalDefines, false))
+		if (!PreprocessShader(PreprocessedShader, Output, Input, AdditionalDefines, EDumpShaderDefines::DontIncludeDefines))
 		{
 			// The preprocessing stage will add any relevant errors.
 			if (Output.Errors.Num() != 0)
diff --git a/Engine/Source/Developer/ShaderPreprocessor/Private/ShaderPreprocessor.cpp b/Engine/Source/Developer/ShaderPreprocessor/Private/ShaderPreprocessor.cpp
--- a/Engine/Source/Developer/ShaderPreprocessor/Private/ShaderPreprocessor.cpp
+++ b/Engine/Source/Developer/ShaderPreprocessor/Private/ShaderPreprocessor.cpp
@@ -219,7 +219,7 @@ bool PreprocessShader(
 	FShaderCompilerOutput& ShaderOutput,
 	const FShaderCompilerInput& ShaderInput,
 	const FShaderCompilerDefinitions& AdditionalDefines,
-	bool bShaderDumpDefinesAsCommentedCode
+	EDumpShaderDefines DefinesPolicy
 	)
 {
 	// Skip the cache system and directly load the file path (used for debugging)
@@ -312,7 +312,7 @@ bool PreprocessShader(
 	}
 
 	// List the defines used for compilation in the preprocessed shaders, especially to know witch permutation vector this shader is.
-	if (bShaderDumpDefinesAsCommentedCode)
+	if (DefinesPolicy == EDumpShaderDefines::AlwaysIncludeDefines || (DefinesPolicy == EDumpShaderDefines::DontCare && ShaderInput.DumpDebugInfoPath.Len() > 0))
 	{
 		DumpShaderDefinesAsCommentedCode(ShaderInput, &OutPreprocessedShader);
 	}
diff --git a/Engine/Source/Developer/ShaderPreprocessor/Public/ShaderPreprocessor.h b/Engine/Source/Developer/ShaderPreprocessor/Public/ShaderPreprocessor.h
--- a/Engine/Source/Developer/ShaderPreprocessor/Public/ShaderPreprocessor.h
+++ b/Engine/Source/Developer/ShaderPreprocessor/Public/ShaderPreprocessor.h
@@ -5,6 +5,19 @@
 #include "CoreMinimal.h"
 #include "ShaderCore.h"
 
+/** Governs the behavior for adding shader defines to the preprocessed source. Can be helpful for the debugging, but makes the source unique
+    which can prevent efficient caching.
+  */
+enum class EDumpShaderDefines : uint8
+{
+	/** Will not be dumped unless Input.DumpDebugInfoPath is set */
+	DontCare,
+	/** No defines */
+	DontIncludeDefines,
+	/** Defines will be added in the comments */
+	AlwaysIncludeDefines
+};
+
 /**
  * Preprocess a shader.
  * @param OutPreprocessedShader - Upon return contains the preprocessed source code.
@@ -19,4 +32,4 @@ extern SHADERPREPROCESSOR_API bool PreprocessShader(
 	FShaderCompilerOutput& ShaderOutput,
 	const FShaderCompilerInput& ShaderInput,
 	const FShaderCompilerDefinitions& AdditionalDefines,
-	bool bShaderDumpDefinesAsCommentedCode = true);
+	EDumpShaderDefines DefinesPolicy = EDumpShaderDefines::DontCare);
