Simple clean up

diff --git a/Engine/Source/Developer/ShaderPreprocessor/Private/ShaderPreprocessor.cpp b/Engine/Source/Developer/ShaderPreprocessor/Private/ShaderPreprocessor.cpp
--- a/Engine/Source/Developer/ShaderPreprocessor/Private/ShaderPreprocessor.cpp
+++ b/Engine/Source/Developer/ShaderPreprocessor/Private/ShaderPreprocessor.cpp
@@ -52,7 +52,7 @@ public:
 		return Loader;
 	}
 
-	bool HasIncludedMendatoryHeaders() const
+	bool HasIncludedMandatoryHeaders() const
 	{
 		return CachedFileContents.Contains(TEXT("/Engine/Public/Platform.ush"));
 	}
@@ -242,7 +242,7 @@ bool PreprocessShader(
 
 	static FCriticalSection McppCriticalSection;
 
-	bool bHasIncludedMendatoryHeaders;
+	bool bHasIncludedMandatoryHeaders = false;
 	{
 		FMcppFileLoader FileLoader(ShaderInput, ShaderOutput);
 
@@ -287,7 +287,7 @@ bool PreprocessShader(
 		McppOutput = McppOutAnsi;
 		McppErrors = McppErrAnsi;
 
-		bHasIncludedMendatoryHeaders = FileLoader.HasIncludedMendatoryHeaders();
+		bHasIncludedMandatoryHeaders = FileLoader.HasIncludedMandatoryHeaders();
 	}
 
 	if (!ParseMcppErrors(ShaderOutput.Errors, ShaderOutput.PragmaDirectives, McppErrors))
@@ -305,7 +305,7 @@ bool PreprocessShader(
 		return false;
 	}
 
-	if (!bHasIncludedMendatoryHeaders)
+	if (!bHasIncludedMandatoryHeaders)
 	{
 		FShaderCompilerError Error;
 		Error.ErrorVirtualFilePath = ShaderInput.VirtualSourceFilePath;
