Replace UE4 reference by simply UE.

#rnx

diff --git a/Engine/Source/Developer/ShaderPreprocessor/Private/PreprocessorPrivate.h b/Engine/Source/Developer/ShaderPreprocessor/Private/PreprocessorPrivate.h
--- a/Engine/Source/Developer/ShaderPreprocessor/Private/PreprocessorPrivate.h
+++ b/Engine/Source/Developer/ShaderPreprocessor/Private/PreprocessorPrivate.h
@@ -29,7 +29,7 @@ inline EMessageType FilterMcppError(const FString& ErrorMsg)
 	};
 	const int32 FilteredSubstringCount = UE_ARRAY_COUNT(SubstringsToFilter);
 
-	if (ErrorMsg.Contains(TEXT("UE4SHADERMETADATA")))
+	if (ErrorMsg.Contains(TEXT("UESHADERMETADATA")))
 	{
 		return EMessageType::ShaderMetaData;
 	}
@@ -46,7 +46,7 @@ inline EMessageType FilterMcppError(const FString& ErrorMsg)
 
 static void ExtractDirective(FString& OutString, FString WarningString)
 {
-	static const FString PrefixString = TEXT("UE4SHADERMETADATA_");
+	static const FString PrefixString = TEXT("UESHADERMETADATA_");
 	uint32 DirectiveStartPosition = WarningString.Find(PrefixString) + PrefixString.Len();
 	uint32 DirectiveEndPosition = WarningString.Find(TEXT("\n"));
 	OutString = WarningString.Mid(DirectiveStartPosition, (DirectiveEndPosition - DirectiveStartPosition));
