Ensure escape characters are encoded correctly in '#line'-directives of shader preprocessor.

#rnx

diff --git a/Engine/Source/Developer/ShaderPreprocessor/Private/ShaderPreprocessor.cpp b/Engine/Source/Developer/ShaderPreprocessor/Private/ShaderPreprocessor.cpp
--- a/Engine/Source/Developer/ShaderPreprocessor/Private/ShaderPreprocessor.cpp
+++ b/Engine/Source/Developer/ShaderPreprocessor/Private/ShaderPreprocessor.cpp
@@ -102,7 +102,8 @@ private:
 			{
 				// Adds a #line 1 "<Absolute file path>" on top of every file content to have nice absolute virtual source
 				// file path in error messages.
-				const FString ShaderFilePath = GetShaderSourceFilePath(VirtualFilePath);
+				FString ShaderFilePath = GetShaderSourceFilePath(VirtualFilePath);
+				ShaderFilePath.ReplaceCharWithEscapedCharInline();
 				FileContents = FString::Printf(TEXT("#line 1 \"%s\"\n%s"), *ShaderFilePath, *FileContents);
 
 				CachedContents = &This->CachedFileContents.Add(VirtualFilePath, StringToArray<ANSICHAR>(*FileContents, FileContents.Len() + 1));
