[shaders] add initial version of stb preprocessor, using default data structure and alloc implementations (for now)

diff --git a/Engine/Source/Developer/ShaderPreprocessor/Private/stb_preprocess/cond_expr.c b/Engine/Source/Developer/ShaderPreprocessor/Private/stb_preprocess/cond_expr.c
new file mode 100644
--- /dev/null
+++ b/Engine/Source/Developer/ShaderPreprocessor/Private/stb_preprocess/cond_expr.c
@@ -0,0 +1,470 @@
+// Copyright Epic Games, Inc. All Rights Reserved.
+
+//////////////////////////////////////////////////////////////////////////////
+//
+// integer constant expression parse & evaluation for #if
+//
+// The text has already been macro-substituted, and defined() was turned
+// into 0 or 1.
+
+#include <ctype.h>
+#include <stddef.h>
+#include <string.h>
+#include <stdlib.h>
+
+#ifdef _MSC_VER
+#define intmax_t __int64
+#define uintmax_t unsigned __int64
+#define strtoull (unsigned) strtoul
+#else
+#define intmax_t long long
+#define uintmax_t unsigned long long
+#endif
+
+typedef struct
+{
+  union {
+    intmax_t  iv;
+    uintmax_t uv;
+  };
+  int is_signed;
+} Value;
+
+enum
+{
+  PPLEX_eof = 0,
+  PPLEX_int_literal,
+  PPLEX_and_and,
+  PPLEX_or_or,
+  PPLEX_less_eq,
+  PPLEX_greater_eq,
+  PPLEX_bang_eq,
+  PPLEX_eq_eq,
+  PPLEX_shift_left,
+  PPLEX_shift_right
+};
+
+typedef struct
+{
+  char *p;
+  int token;  
+  Value token_value;
+  int syntax_error;
+} ppcexp; // preprocessor constant expression lexer
+
+static void ce_next(ppcexp *c)
+{
+  char *p = c->p;
+  while (isspace(*p))
+    ++p;
+  switch (*p) {
+    case '<': if (p[1] == '<') {
+                c->token = PPLEX_shift_left;
+                p += 2;
+              } else if (p[1] == '=') {
+                c->token = PPLEX_less_eq;
+                p += 2;
+              } else {
+                c->token = *p++;
+              }
+              break;
+    case '>': if (p[1] == '>') {
+                c->token = PPLEX_shift_right;
+                p += 2;
+              } else if (p[1] == '=') {
+                c->token = PPLEX_greater_eq;
+                p += 2;
+              } else {
+                c->token = *p++;
+              }
+              break;
+    case '&': if (p[1] == '&') {
+                c->token = PPLEX_and_and;
+                p += 2;
+              } else {
+                c->token = *p++;
+              }
+              break;
+    case '|': if (p[1] == '|') {
+                c->token = PPLEX_or_or;
+                p += 2;
+              } else {
+                c->token = *p++;
+              }
+              break;
+    case '=': if (p[1] == '=') {
+                c->token = PPLEX_eq_eq;
+                p += 2;
+              } else {
+                c->token = *p++;
+              }
+              break;
+    case '!': if (p[1] == '=') {
+                c->token = PPLEX_bang_eq;
+                p += 2;
+              } else {
+                c->token = *p++;
+              }
+              break;
+    case 0:
+      c->token = 0;
+      break;
+    default:
+      c->token = *p++;
+      break;
+
+    case 'a': case 'b': case 'c': case 'd': case 'e': case 'f': case 'g': case 'h': case 'i':
+    case 'j': case 'k': case 'l': case 'm': case 'n': case 'o': case 'p': case 'q': case 'r':
+    case 's': case 't': case 'u': case 'v': case 'w': case 'x': case 'y': case 'z':
+    case 'A': case 'B': case 'C': case 'D': case 'E': case 'F': case 'G': case 'H': case 'I':
+    case 'J': case 'K': case 'L': case 'M': case 'N': case 'O': case 'P': case 'Q': case 'R':
+    case 'S': case 'T': case 'U': case 'V': case 'W': case 'X': case 'Y': case 'Z':
+    case '_': case '$':
+    {
+      do
+        ++p;
+      while (isalnum(*p) || *p == '_' || *p == '$');
+      // any Identifier is a macro that WASN'T defined, so it evaluates as 0
+      c->token_value.iv = 0;
+      c->token_value.is_signed = 1;
+      c->token = PPLEX_int_literal;
+      break;
+    }
+
+    case '0':
+      if (p[1] == 'x' || p[1] == 'X') {
+        char *q;
+        c->token_value.uv = strtoull(p+2, &q, 16);
+        if (q == p+2) c->syntax_error = 1;
+        p = q;
+      } else if (p[1] == 'b' || p[1] == 'B') {
+        char *q;
+        c->token_value.uv = strtoull(p+2, &q, 2);
+        if (q == p+2) c->syntax_error = 1;
+        p = q;
+      } else if (p[1] == '.') {
+        c->syntax_error = 1;
+      } else {
+        c->token_value.uv = strtoull(p, &p, 8);
+      }
+      goto numeric_suffixes;
+    
+    case '1': case '2': case '3': case '4':
+    case '5': case '6': case '7': case '8': case '9':
+      c->token_value.uv = strtoull(p, &p, 10);
+    numeric_suffixes:
+      c->token_value.is_signed = (c->token_value.uv < ((uintmax_t) 1 << 63));
+      while (*p == 'u' || *p == 'l') {
+        if (*p == 'u') c->token_value.is_signed = 1;
+        ++p;
+      }
+      c->token = PPLEX_int_literal;
+      break;
+  }
+  c->p = p;
+}
+
+static Value ce_primary(ppcexp *c)
+{
+  Value v;
+  if (c->token != PPLEX_int_literal)
+    c->syntax_error = 1;
+  v = c->token_value;
+  ce_next(c);
+  return v;
+}
+
+static Value ce_cond(ppcexp *c);
+static Value ce_unary(ppcexp *c)
+{
+  Value v;
+  switch (c->token) {
+    case '+':
+      ce_next(c);
+      return ce_unary(c);
+    case '-':
+      ce_next(c);
+      v = ce_unary(c);
+      v.iv = -v.iv;
+      return v;
+    case '!':
+      ce_next(c);
+      v = ce_unary(c);
+      v.iv = (v.uv == 0);
+      v.is_signed = 1;
+      return v;
+    case '~':
+      ce_next(c);
+      v = ce_unary(c);
+      v.uv = ~v.uv;
+      return v;
+    case '(':
+      ce_next(c);
+      v = ce_cond(c);
+      if (c->token != ')')
+        c->syntax_error = 1;
+      ce_next(c);
+      return v;
+    default:
+      return ce_primary(c);
+  }
+}
+
+static Value ce_mul(ppcexp *c)
+{
+  Value left = ce_unary(c);
+  while (c->token == '*' || c->token == '/' || c->token == '%') {
+    int token = c->token;
+    Value right;
+    ce_next(c);
+    right = ce_unary(c);
+    if (token == '*')
+      left.uv *= right.uv;
+    else {
+      if (right.uv == 0) {
+        if (c->syntax_error == 0)
+          c->syntax_error = -1;
+      } else if (left.is_signed & right.is_signed)
+        if (token == '/')
+          left.iv /= right.iv;
+        else
+          left.iv %= right.iv;
+      else
+        if (token == '/')
+          left.uv /= right.uv;
+        else
+          left.uv %= right.uv;
+      left.is_signed &= right.is_signed;
+    }
+  }
+  return left;
+}
+
+static Value ce_sum(ppcexp *c)
+{
+  Value left = ce_mul(c);
+  while (c->token == '+' || c->token == '-') {
+    int token = c->token;
+    Value right;
+    ce_next(c);
+    right = ce_mul(c);
+    if (token == '+')
+      left.uv += right.uv;
+    else
+      left.uv -= right.uv;
+    left.is_signed &= right.is_signed;
+  }
+  return left;
+}
+
+static Value ce_shift(ppcexp *c)
+{
+  Value left = ce_sum(c);
+  while (c->token == PPLEX_shift_left || c->token == PPLEX_shift_right) {
+    int tok = c->token;
+    Value right;
+    ce_next(c);
+    right = ce_sum(c);
+    if (tok == PPLEX_shift_right)
+      if (left.is_signed)
+        left.iv = (left.iv >> right.uv);
+      else
+        left.uv = (left.uv >> right.uv);
+    else
+      left.uv = (left.uv << right.uv);
+  }
+  return left;
+}
+
+static Value ce_compare_inequality(ppcexp *c)
+{
+  Value left = ce_shift(c);
+  for(;;) {
+    Value right;
+    switch (c->token) {
+      case PPLEX_less_eq:
+        ce_next(c);
+        right = ce_shift(c);
+        left.iv = (left.is_signed & right.is_signed) ? (left.iv <= right.iv) : (left.uv <= right.uv);
+        break;
+      case PPLEX_greater_eq:
+        ce_next(c);
+        right = ce_shift(c);
+        left.iv = (left.is_signed & right.is_signed) ? (left.iv >= right.iv) : (left.uv >= right.uv);
+        break;
+      case '<':
+        ce_next(c);
+        right = ce_shift(c);
+        left.iv = (left.is_signed & right.is_signed) ? (left.iv <  right.iv) : (left.uv <  right.uv);
+        break;
+      case '>':
+        ce_next(c);
+        right = ce_shift(c);
+        left.iv = (left.is_signed & right.is_signed) ? (left.iv >  right.iv) : (left.uv >  right.uv);
+        break;
+      default:
+        return left;
+    }
+    left.is_signed = 1;
+  }
+}
+
+static Value ce_compare_equality(ppcexp *c)
+{
+  Value left = ce_compare_inequality(c), right;
+  while (c->token == PPLEX_eq_eq || c->token == PPLEX_bang_eq) {
+    int token = c->token;
+    ce_next(c);
+    right = ce_compare_inequality(c);
+    left.uv = (left.uv == right.uv) ^ (token == PPLEX_bang_eq);
+    left.is_signed = 1;
+  }
+  return left;
+}
+
+static Value ce_bitwise_and(ppcexp *c)
+{
+  Value left = ce_compare_equality(c), right;
+  while (c->token == '&') {
+    ce_next(c);
+    right = ce_compare_equality(c);
+    left.uv &= right.uv;
+    left.is_signed &= right.is_signed;
+  }
+  return left;
+}
+
+static Value ce_bitwise_xor(ppcexp *c)
+{
+  Value left = ce_bitwise_and(c), right;
+  while (c->token == '^') {
+    ce_next(c);
+    right = ce_bitwise_and(c);
+    left.uv ^= right.uv;
+    left.is_signed &= right.is_signed;
+  }
+  return left;
+}
+
+static Value ce_bitwise_or(ppcexp *c)
+{
+  Value left = ce_bitwise_xor(c), right;
+  while (c->token == '|') {
+    ce_next(c);
+    right = ce_bitwise_xor(c);
+    left.uv |= right.uv;
+    left.is_signed &= right.is_signed;
+  }
+  return left;
+}
+
+static Value ce_logical_and(ppcexp *c)
+{
+  Value left = ce_bitwise_or(c), right;
+  while (c->token == PPLEX_and_and) {
+    int err = c->syntax_error;
+    ce_next(c);
+    right = ce_bitwise_or(c);
+    if (left.uv == 0 && c->syntax_error < 0 && err == 0)
+      c->syntax_error = 0; // suppress unevaluated div by 0
+    left.uv = left.uv && right.uv;
+    left.is_signed = 1;
+  }
+  return left;
+}
+
+static Value ce_logical_or(ppcexp *c)
+{
+  Value left = ce_logical_and(c), right;
+  while (c->token == PPLEX_or_or) {
+    int err = c->syntax_error;
+    ce_next(c);
+    right = ce_logical_and(c);
+    if (left.uv != 0 && c->syntax_error < 0 && err == 0)
+      c->syntax_error = 0; // suppress unevaluated div by 0
+    left.uv = left.uv || right.uv;
+    left.is_signed = 1;
+  }
+  return left;
+}
+
+static Value ce_cond(ppcexp *c)
+{
+  Value cond = ce_logical_or(c);
+  if (c->token != '?')
+    return cond;
+  else {
+    int err;
+    Value result, left;
+    err = c->syntax_error;
+    ce_next(c);
+    left = ce_cond(c);
+    if (c->token != ':') {
+      c->syntax_error = 1;
+      return cond;
+    }
+    if (cond.uv == 0 && c->syntax_error < 0 && err == 0)
+      c->syntax_error = 0; // suppress unevaluated div by 0
+    err = c->syntax_error;
+    ce_next(c);
+    result = ce_cond(c);
+    if (cond.uv != 0 && c->syntax_error < 0 && err == 0)
+      c->syntax_error = 0; // suppress unevaluated div by 0
+    result.uv = cond.uv != 0 ? left.uv : result.uv;
+    result.is_signed &= left.is_signed;
+    return result;
+  }  
+}
+
+int evaluate_integer_constant_expression_as_condition(char *p, int *syntax_error)
+{
+  Value v;
+  ppcexp c = { 0 };
+  c.p = p;
+  ce_next(&c);
+  v = ce_cond(&c);
+  if (c.token != PPLEX_eof)
+    c.syntax_error = 1;
+  *syntax_error = c.syntax_error;
+  return v.uv != 0;
+}
+
+#ifdef _MSC_VER
+__int64
+#else
+long long
+#endif
+evaluate_integer_constant_expression(char *p, int *syntax_error)
+{
+  Value v;
+  ppcexp c = { 0 };
+  c.p = p;
+  ce_next(&c);
+  v = ce_cond(&c);
+  if (c.token != PPLEX_eof)
+    c.syntax_error = 1;
+  *syntax_error = c.syntax_error;
+  return v.iv;
+}
+
+#include <assert.h>
+void test_integer_constant_expression(void)
+{
+  int err;
+  assert(evaluate_integer_constant_expression("4+3*2",&err)==10);              assert(err==0);
+  assert(evaluate_integer_constant_expression("4+3*2==10",&err) == 1);         assert(err==0);
+  assert(evaluate_integer_constant_expression("0xffffffff+1==1<<32",&err)==1); assert(err==0);
+  evaluate_integer_constant_expression("1+=1",&err); assert(err>0);
+  evaluate_integer_constant_expression("1+(",&err) ; assert(err>0);
+  evaluate_integer_constant_expression("(1+",&err) ; assert(err>0);
+  assert(evaluate_integer_constant_expression("2/1",&err)==2); assert(err==0);
+  evaluate_integer_constant_expression("1/0",&err); assert(err<0);
+  evaluate_integer_constant_expression("1 || 1/0",&err); assert(err==0);
+  evaluate_integer_constant_expression("0 || 1/0",&err); assert(err<0);
+  evaluate_integer_constant_expression("1 && 1/0",&err); assert(err<0);
+  evaluate_integer_constant_expression("0 && 1/0",&err); assert(err==0);
+  evaluate_integer_constant_expression("0 ? 1/0 : 2/1",&err); assert(err==0);
+  evaluate_integer_constant_expression("1 ? 1/0 : 2/1",&err); assert(err<0);
+  evaluate_integer_constant_expression("0 ? 2/1 : 1/0",&err); assert(err<0);
+  evaluate_integer_constant_expression("1 ? 2/1 : 1/0",&err); assert(err==0);
+}
diff --git a/Engine/Source/Developer/ShaderPreprocessor/Private/stb_preprocess/cond_expr.h b/Engine/Source/Developer/ShaderPreprocessor/Private/stb_preprocess/cond_expr.h
new file mode 100644
--- /dev/null
+++ b/Engine/Source/Developer/ShaderPreprocessor/Private/stb_preprocess/cond_expr.h
@@ -0,0 +1,10 @@
+// Copyright Epic Games, Inc. All Rights Reserved.
+
+#ifdef _MSC_VER
+__int64
+#else
+long long
+#endif
+            evaluate_integer_constant_expression(char *p, int *syntax_error);
+extern int  evaluate_integer_constant_expression_as_condition(char *p, int *syntax_error);
+extern void test_integer_constant_expression(void);
diff --git a/Engine/Source/Developer/ShaderPreprocessor/Private/stb_preprocess/preprocessor.c b/Engine/Source/Developer/ShaderPreprocessor/Private/stb_preprocess/preprocessor.c
new file mode 100644
--- /dev/null
+++ b/Engine/Source/Developer/ShaderPreprocessor/Private/stb_preprocess/preprocessor.c
@@ -0,0 +1,4271 @@
+// Copyright Epic Games, Inc. All Rights Reserved.
+
+#define _CRT_SECURE_NO_WARNINGS
+#include <stdio.h>
+#include <stdlib.h>
+#include <string.h>
+#include <stdarg.h>
+#include <assert.h>
+#include <ctype.h>
+#include <stdarg.h>
+#include "preprocessor.h"
+#include "cond_expr.h"
+
+#ifdef PPTEST
+#define STB_DS_IMPLEMENTATION
+#define STB_ALLOC_IMPLEMENTATION
+#endif
+
+#include "stb_ds.h"
+#include "stb_alloc.h"
+
+#if 0
+#include <crtdbg.h>
+#define MDBG()   _CrtCheckMemory()
+#else
+#define MDBG()
+#endif
+
+//////////////////////////////////////////////////////////////////////////////
+//
+//  Limits to prevent infinite recursion
+//
+
+#define MAX_INCLUDE_NESTING    200
+#define MAX_MACRO_NESTING      5000
+#define MAX_FILENAME           4096
+
+//////////////////////////////////////////////////////////////////////////////
+//
+//  TYPES
+//
+
+typedef unsigned char    uint8;
+typedef unsigned int     uint;
+typedef unsigned int     uint32;
+typedef unsigned short   uint16;
+typedef   signed short   int16;
+typedef          int     ppbool;
+#ifdef _MSC_VER
+typedef unsigned __int64 uint64;
+#else
+#include <stdint.h>
+typedef uint64_t uint64;
+#define _vsnprintf vsnprintf
+#define _strdup strdup
+#endif
+
+// macro expansion helper struct
+typedef struct
+{
+  char   *text;
+  uint32  text_length;
+  int16   argument_index;
+  uint8   stringify_argument;
+  uint8   eliminate_comma_before_empty_arg;  // gcc extension using ## to fix up __VA_ARGS__
+  uint8   concat_argument;
+} expansion_part;
+
+struct macro_definition
+{
+  char  * symbol_name;
+  int      symbol_name_length;
+  union {
+    char           *simple_expansion;
+    expansion_part *expansion; // array
+  };
+  int      simple_expansion_length;
+  int      num_parameters; // if num_parameters < 0, then it doesn't take parentheses when used
+  uint8    predefined;
+  uint8    is_variadic;    // if true, then __VA_ARGS__ will be specified using 'num_arguments' as argument_index
+  uint8    disabled;
+  macro_definition *next; // push/pop stack
+};
+
+#define MACRO_NUM_PARAMETERS_no_parentheses  -1
+#define MACRO_NUM_PARAMETERS_file            -2
+#define MACRO_NUM_PARAMETERS_line            -3
+#define MACRO_NUM_PARAMETERS_defined         -4
+#define MACRO_NUM_PARAMETERS_counter         -5
+
+static macro_definition predefined_FILE    = { "__FILE__"   ,  8, { 0 }, 0, MACRO_NUM_PARAMETERS_file   , 1 };
+static macro_definition predefined_LINE    = { "__LINE__"   ,  8, { 0 }, 0, MACRO_NUM_PARAMETERS_line   , 1 };
+static macro_definition predefined_COUNTER = { "__COUNTER__", 11, { 0 }, 0, MACRO_NUM_PARAMETERS_counter, 1 };
+static macro_definition predefined_defined = { "defined"    ,  7, { 0 }, 0, MACRO_NUM_PARAMETERS_defined, 1 };
+
+#if defined(_MSC_VER) && _MSC_VER <= 1400
+  #define strdup(x)  strdup(x)
+#elif !defined(strdup)
+  #define strdup(x)  _strdup(x)
+#endif
+
+typedef struct
+{
+  char             * key;
+  macro_definition * value;
+} macro_hash_entry;
+
+// perfect hash table for directives
+#define DIRECTIVE_HASH_SIZE   32
+struct {
+  char *name;
+  size_t name_len;
+} directive_hash[DIRECTIVE_HASH_SIZE];
+
+typedef struct
+{
+  uint32 line;
+  uint32 type;
+} ifdef_info;
+
+typedef struct
+{
+  char *key;   // filename
+  char *value; // macro name
+} include_once;
+
+typedef struct
+{
+  char *key;
+  void *value;
+  size_t key_length;
+} hashtable_pair; // 24 bytes
+
+typedef struct
+{
+  uint32 hash;
+  uint32 index;
+} table_entry;
+
+typedef struct
+{
+  char *key;
+  void *value;
+} shpair;
+
+#define HASH_EMPTY_MARKER  0
+#define HASH_TOMBSTONE    1
+
+typedef struct
+{
+  table_entry    *table ; // malloc
+  hashtable_pair *pair  ; // stb_ds array
+  int             first_unused_pair;
+  uint32          table_size;
+  uint32          table_mask;
+  uint32          count_in_use;
+  uint32          count_tombstones;
+  uint32          doubling_threshold;
+  stb_arena       arena;   
+  shpair         *verify;
+} pphash;
+
+typedef struct pp_context 
+{
+  char **include_paths;
+  int    num_include_paths;
+  char **sys_include_paths;
+  int    num_sys_include_paths;
+  pphash macro_map;                // map of macro names to their definitions
+  macro_hash_entry *undef_map;     // map of macro names to a flag for whether they've ever been undef'd, for more useful error messages
+  uint8 (*macro_name_filter)[256]; // 64KB table indicating whether 1 & 2-character identifiers are macros
+  include_once     *once_map;      // map of which filenames have been #onced
+
+  int include_nesting_level;       // we stop after a certain number in case of unbounded recursive includes
+  int macro_expansion_level;       // because recursive macros are prevented, we don't need to stop this, but we do in case of bugs in recursive macro processing
+  char *last_ifndef;
+  macro_definition *last_macro_definition;
+  int include_guard_detect;
+
+  ifdef_info *ifdef_stack;
+
+  pp_where *include_stack;
+  pp_diagnostic *error;
+  int            stop; // if we hit an error that causes us to stop processing
+
+  int num_lines;
+  int num_disabled_lines;
+  int num_onced_files;
+  int num_includes;
+  int counter;
+
+  stb_arena macro_arena;
+} pp_context;
+
+typedef struct parse_state
+{
+  pp_context *context;
+
+  char   * filename; // this is cached in last_output_filename so must remain valid indefinitely
+  char   * last_output_filename;
+  char   * stringified_filename; // this is NULL until first time __FILE__ is used, then it's cached in macro_arena
+
+  char   * src;
+  size_t   src_offset;
+  size_t   src_length;
+
+  char   * dest;
+  char   * fast_dest;
+
+  int             src_line_number;
+  int             dest_line_number;
+  int             num_disabled_lines;
+  int             conditional_nesting_depth_at_start;
+  int             state_limit;
+
+  int             include_guard_endif_level;
+  char          * include_guard_candidate_macro_name; // this is also the flag for whether we're searching for a matching #endif
+
+  unsigned char   in_leading_whitespace;
+  unsigned char   system_include;
+
+  struct parse_state    *parent;
+} parse_state;
+
+//////////////////////////////////////////////////////////////////////////////
+//
+//  UTILITIES
+//
+
+static char *arena_alloc_padded_string(stb_arena *a, char *text, size_t stringlen)
+{
+  // allocate a string with whitespace at each end to prevent accidental token-pasting
+  size_t final_length = stringlen + 3;
+  char *mem = stb_arena_alloc_aligned(a, final_length, 4);
+  mem[0] = ' ';
+  memcpy(mem+1, text, stringlen);
+  mem[stringlen+1] = ' ';
+  mem[stringlen+2] = 0;
+  return mem;
+}
+
+#define MIN3(a, b, c) ((a) < (b) ? ((a) < (c) ? (a) : (c)) : ((b) < (c) ? (b) : (c)))
+
+static int compute_edit_distance(const char *s1, const char *s2, int max_dist)
+{
+  unsigned int s1len, s2len, x, y, lastdiag, olddiag;
+  unsigned int column[81];
+  s1len = (unsigned int) strlen(s1);
+  s2len = (unsigned int) strlen(s2);
+  if (s2len < s1len) {
+    unsigned int ti;
+    const char *ts;
+    ti = s1len;
+    s1len = s2len;
+    s2len = ti;
+    ts = s1;
+    s1 = s2;
+    s2 = ts;
+  }
+  if (s1len+1 > sizeof(column))
+    return 9999;
+  if (s2len > s1len + max_dist)
+    return 9999;
+  for (y = 1; y <= s1len; y++)
+      column[y] = y;
+  for (x = 1; x <= s2len; x++) {
+      column[0] = x;
+      for (y = 1, lastdiag = x-1; y <= s1len; y++) {
+          olddiag = column[y];
+          column[y] = MIN3(column[y] + 1, column[y-1] + 1, lastdiag + (s1[y-1] == s2[x-1] ? 0 : 1));
+          lastdiag = olddiag;
+      }
+  }
+  return(column[s1len]);
+}
+
+static void find_closest_match_init(int *dist, char **best, int max_dist)
+{
+  *dist = max_dist+1;
+  *best = 0;
+}
+
+static void find_closest_match_step(char *reference, char *varying, int *a, char **best)
+{
+  int dist = compute_edit_distance(reference, varying, *a);
+  if (dist < *a) {
+    *a = dist;
+    *best = varying;
+  }
+}
+
+static char *find_closest_matching_string(char *reference, char *choices[], int num_choices, int max_dist)
+{
+  int i, dist;
+  char *best;
+
+  find_closest_match_init(&dist, &best, max_dist);
+
+  for (i=0; i < num_choices; ++i)
+    find_closest_match_step(reference, choices[i], &dist, &best);
+
+  if (dist <= max_dist)
+    return best;
+  else
+    return NULL;
+}
+
+static char *error_char(char ch, char buffer[12])
+{
+  if (ch >= 32 && ch < 127)
+    if (ch == '\'' || ch == '\\')
+      sprintf(buffer, "'\\%c'", ch);
+    else
+      sprintf(buffer, "'%c'", ch);
+  else if (ch == 0)
+    sprintf(buffer, "EOF");
+  else
+    sprintf(buffer, "value %d", ch);  // "value -100\0" = 11 (6 + 1 + 3 + 1)
+  return buffer;
+}
+
+// write a line number in 12 characters of space, and return a pointer after the number
+// because this is used for #line output, we do a few special things:
+//     1. 1-digit numbers take 3 characters with two leading spaces
+//     2. 2-digit numbers take 2 characters with one leading space
+//     3. all other numbers have no extra spaces
+//     4. sprintf() is used if number is > 9999, on the assumption this is rare
+//            -- but note that if a file is, say, 500K lines, then there may be
+//               very few FILES that use this path, but an inordinate number of lines
+static char * preprocessor_line_dtoa(char q[12], int number)
+{
+  static const char digits[10+1] = "0123456789";
+  static const char digit_pairs[200+1] =
+     "0001020304050607080910111213141516171819202122232425262728293031323334353637383940414243444546474849"
+     "5051525354555657585960616263646566676869707172737475767778798081828384858687888990919293949596979899";
+
+  if (number <= 999) {
+    static const char first_digit[10+1] = " 123456789";
+    int n = number % 100;
+    int leading = number / 100;
+    q[0] = first_digit[leading];
+    q[1] = digit_pairs[n*2+0];
+    q[2] = digit_pairs[n*2+1];
+    #if 1
+    if (number < 10)
+      q[1] = ' ';
+    #else
+    q[(number < 10) ? 1 : 3] = ' ';
+    #endif
+    return q+3;
+  } else if (number <= 9999) {
+    int second = number % 100;
+    int first = number / 100;
+    q[0] = digit_pairs[first*2+0];
+    q[1] = digit_pairs[first*2+1];
+    q[2] = digit_pairs[second*2+0];
+    q[3] = digit_pairs[second*2+1];
+    return q+4;
+  } else {
+    return q+sprintf(q, "%d", number);
+  }
+}
+
+static void output_line_directive(parse_state *cs)
+{
+  char *q;
+  char *out = cs->dest;
+
+  // output the filename if it's changed
+  char *filename = (cs->last_output_filename && cs->last_output_filename == cs->filename ? "" : cs->filename);
+
+  // compute an upper bound on the amount of data to be written
+  size_t filename_len = strlen(filename);
+  size_t expansion_length = 6 + 16 + filename_len + 2 + 2 + 1;
+
+  if (cs->fast_dest)
+     q = cs->fast_dest;
+  else {
+    // make sure there's room in the output for the #line directive plus all the remaining text
+    arrsetcap(out, arrlen(out) + (size_t) (cs->src_length - cs->src_offset) + expansion_length);
+
+    q = out + arrlen(out);
+  }
+
+  // write the line directive
+  *q++ = '#';
+  *q++ = ' ';
+  q = preprocessor_line_dtoa(q, cs->src_line_number);
+  if (filename[0]) {
+    *q++ = ' ';
+    *q++ = '"';
+    strcpy(q, filename); // @TODO cache stringified filename
+    q += filename_len;
+    *q++ = '"';
+    cs->last_output_filename = filename;
+  } 
+  *q++ = '\n';
+
+
+  cs->dest_line_number = cs->src_line_number;
+  if (cs->fast_dest)
+    cs->fast_dest = q;
+  else {
+    arrsetlen(out, (q - out));
+    cs->dest = out;
+  }
+}
+
+//////////////////////////////////////////////////////////////////////////////
+//
+//  ERROR HANDLING
+//
+
+static uint8 pp_result_mode[PP_RESULT_count] =
+{
+   PP_RESULT_MODE_no_warning,
+   PP_RESULT_MODE_supplementary,
+   PP_RESULT_MODE_no_warning, // #undef of undefined symbol is silently ignored per spec
+   // initialize all others to PP_RESULT_MODE_stop
+}; 
+
+void pp_set_warning_mode(int result_code, int result_mode)
+{
+  pp_result_mode[result_code] = (uint8) result_mode;
+}
+
+static int did_you_mean_threshold = PP_RESULT_MODE_warning;
+
+void pp_set_did_you_mean_threshold(int mode)
+{
+  did_you_mean_threshold = mode;
+}
+
+static int error_explicit_v(parse_state *ps, int code, int line_number, char *text, va_list va)
+{
+  parse_state *stack;
+  pp_where where;
+  pp_diagnostic d;
+  size_t len;
+
+  #if defined(_MSC_VER) && _MSC_VER < 1500
+  len = strlen(text) + 256;
+  #else
+  {
+    va_list vc;
+    #ifdef _MSC_VER
+    vc = va;
+    #else
+    va_copy(vc, va);
+    #endif
+    len = _vsnprintf(0, 0, text, vc);
+    va_end(vc);
+  }
+  #endif
+
+  d.message = malloc(len+1);
+  len = _vsnprintf(d.message, len, text, va);
+  d.message[len] = 0;
+
+  d.diagnostic_code = code;
+  d.error_level     = pp_result_mode[code];
+  if (d.error_level == PP_RESULT_MODE_warning_fast)
+    d.error_level = PP_RESULT_MODE_warning;
+  d.where = 0;
+
+  stack = ps;
+  while (stack != NULL) {
+    if (stack->filename) {
+      where.filename = strdup(stack->filename);
+      where.line_number = (line_number > 0 ? line_number : stack->src_line_number);
+      where.column = 0;
+      line_number = 0;
+      arrpush(d.where, where);
+    }
+    stack = stack->parent;
+  }
+
+  arrpush(ps->context->error, d);
+
+  if (d.error_level == PP_RESULT_MODE_error)
+    ps->context->stop = 1;
+
+  return d.error_level == PP_RESULT_MODE_error;
+}
+
+// returns TRUE if it should terminate
+static int do_error_code(parse_state *ps, int code, char *text, ...)
+{
+  int err;
+  va_list va;
+
+  if (pp_result_mode[code] == PP_RESULT_MODE_no_warning)
+    return 0;
+
+  va_start(va, text);
+  err = error_explicit_v(ps, code, 0, text, va);
+  va_end(va);
+  return err;
+}
+
+static int do_error(parse_state *ps, char *text, ...)
+{
+  int err;
+  va_list va;
+
+  va_start(va, text);
+  err = error_explicit_v(ps, PP_RESULT_ERROR, 0, text, va);
+  va_end(va);
+  return err;
+}
+
+static int do_error_explicit(parse_state *ps, int code, int line_number, char *text, ...)
+{
+  int err;
+  va_list va;
+
+  if (pp_result_mode[code] == PP_RESULT_MODE_no_warning)
+    return 0;
+
+  va_start(va, text);
+  err = error_explicit_v(ps, code, line_number, text, va);
+  va_end(va);
+  return err;
+}
+
+static void error_supplement(parse_state *cs, int code, int line_number, char *text, ...)
+{
+  va_list va;
+  if (pp_result_mode[code] == PP_RESULT_MODE_no_warning)
+    return;
+  va_start(va, text);
+  error_explicit_v(cs, PP_RESULT_supplementary, line_number, text, va);
+  va_end(va);
+}
+
+//////////////////////////////////////////////////////////////////////////////
+//
+//  STRING HASHING
+//
+//  based on XXHash32
+
+static unsigned int load_unaligned_32(void *ptr)
+{
+#ifdef NO_UNALIGNED_LOADS
+  unsigned char *data = ptr;
+  return data[0] | (data[1] << 8) | (data[2] << 16) | (data[3] << 24); 
+#else
+  unsigned int *data = ptr; // DOESN'T produce the same result for little-endian & big-endian
+  return *data;
+#endif
+}
+
+#define rotl(x,y)   (((x) << y) | ((x) >> (-(y)&31)))
+#define xxprime1 0x9E3779B1u
+#define xxprime2 0x85EBCA77u
+#define xxprime3 0xC2B2AE3Du
+
+// export hash function so compiler can use it as well
+unsigned int preprocessor_hash(char *data, size_t len)
+{
+  unsigned int hash;
+  if (len < 4) {
+    // load 3 bytes, overlapping if len < 3
+    static unsigned char offset1[4] = { 0,0,1,1 };
+    static unsigned char offset2[4] = { 0,0,0,2 };
+    unsigned int h = data[0] + (data[offset1[len]]<<8) + (data[offset2[len]]<<16);
+    h = xxprime1 + h*xxprime2;
+    hash = rotl(h,13) * xxprime1;
+  } else if (len <= 8) {
+    // load two 32-bit ints, overlapping if len < 8
+    unsigned int h0 = xxprime1 + load_unaligned_32(data+  0  )*xxprime2;
+    unsigned int h1 = xxprime2 + load_unaligned_32(data+len-4)*xxprime2;
+    h0 = rotl(h0,13) * xxprime1;
+    h1 = rotl(h1,13) * xxprime1;
+    hash = rotl(h0,1) + rotl(h1,7);
+  } else {
+    unsigned int h0 = xxprime1 ^ xxprime2;
+    unsigned int h1 = xxprime2;
+    unsigned int h2 = 0;
+    unsigned int h3 = xxprime1;
+    if (len < 16) {
+      // load four 32-bit ints, partially overlapping
+      h0 = h0 + load_unaligned_32(data+    0) * xxprime2;
+      h1 = h1 + load_unaligned_32(data+    4) * xxprime2;
+      h2 = h2 + load_unaligned_32(data+len-8) * xxprime2;
+      h3 = h3 + load_unaligned_32(data+len-4) * xxprime2;
+      h0 = rotl(h0,13)*xxprime1;
+      h1 = rotl(h1,13)*xxprime1;
+      h2 = rotl(h2,13)*xxprime1;
+      h3 = rotl(h3,13)*xxprime1;
+    } else {
+      while (len > 16) {
+        h0 = h0 + load_unaligned_32(data+ 0) * xxprime2;
+        h1 = h1 + load_unaligned_32(data+ 4) * xxprime2;
+        h2 = h2 + load_unaligned_32(data+ 8) * xxprime2;
+        h3 = h3 + load_unaligned_32(data+12) * xxprime2;
+        h0 = rotl(h0,13)*xxprime1;
+        h1 = rotl(h1,13)*xxprime1;
+        h2 = rotl(h2,13)*xxprime1;
+        h3 = rotl(h3,13)*xxprime1;
+        data += 16;
+        len -= 16;
+      }
+      data = data + len - 16;
+      h0 = h0 + load_unaligned_32(data+ 0) * xxprime2;
+      h1 = h1 + load_unaligned_32(data+ 4) * xxprime2;
+      h2 = h2 + load_unaligned_32(data+ 8) * xxprime2;
+      h3 = h3 + load_unaligned_32(data+12) * xxprime2;
+      h0 = rotl(h0,13)*xxprime1;
+      h1 = rotl(h1,13)*xxprime1;
+      h2 = rotl(h2,13)*xxprime1;
+      h3 = rotl(h3,13)*xxprime1;
+    }
+    hash = rotl(h0,1)+rotl(h1,7)+rotl(h2,12)+rotl(h3,18);
+  }
+  hash += (unsigned int) len;
+  // XXH hash avalanche
+  hash ^= hash >> 15;
+  hash *= xxprime2;
+  hash ^= hash >> 13;
+  hash *= xxprime3;
+  hash ^= hash >> 16;
+  return hash;
+}
+
+void stringhash_init_table(pphash *ht, unsigned int size)
+{
+  table_entry t = { 0 }, *table;
+  unsigned int i;
+  table = malloc(sizeof(table[0]) * size);
+  ht->table_size = size;
+  ht->table_mask = size-1;
+  ht->table      = table;
+  ht->doubling_threshold = size - (size >> 2); // 75%
+  for (i=0; i < size; ++i)
+    table[i] = t; // hope for 8-byte-at-a-time init
+}
+
+void *stringhash_get(pphash *ht, char *key, size_t keylen)
+{
+  uint32 mask = ht->table_mask;
+  uint32 hash = preprocessor_hash(key, keylen);
+  uint32 slot = hash & mask;
+  uint32 step;
+  void * vresult=0;
+  table_entry *table = ht->table;
+
+  if (ht->verify) { char temp[1024]; memcpy(temp, key, keylen); temp[keylen] = 0; vresult = shget(ht->verify, temp); }
+
+  if (hash <= HASH_TOMBSTONE) hash+=2;
+  if (table[slot].hash == hash) {
+    uint32 i = table[slot].index;
+    if (ht->pair[i].key_length == keylen && 0==memcmp(key, ht->pair[i].key, keylen)) {
+      if (ht->verify) assert(vresult == ht->pair[i].value);
+      return ht->pair[i].value;
+    }
+  }
+  if (table[slot].hash == HASH_EMPTY_MARKER) {
+    assert(vresult == NULL);
+    return NULL;
+  }
+
+  step = hash ^ (hash >> 16) ^ (hash >> 23);
+  step |= 1;
+  for(;;) {
+    slot = (slot + step) & mask;
+    if (table[slot].hash == hash) {
+      uint32 i = table[slot].index;
+      if (ht->pair[i].key_length == keylen && 0==memcmp(key, ht->pair[i].key, keylen)) {
+        if (ht->verify) assert(vresult == ht->pair[i].value);
+        return ht->pair[i].value;
+      }
+    }
+    if (table[slot].hash == HASH_EMPTY_MARKER) {
+      assert(vresult == NULL);
+      return NULL;
+    }
+  }
+}
+
+int stringhash_delete(pphash *ht, char *key, size_t keylen)
+{
+  uint32 mask = ht->table_mask;
+  uint32 hash = preprocessor_hash(key, keylen);
+  uint32 slot = hash & mask;
+  uint32 step, i;
+  table_entry *table = ht->table;
+
+  if (ht->verify) { char temp[1024]; memcpy(temp, key, keylen); temp[keylen] = 0; shdel(ht->verify, temp); }
+
+  if (hash <= HASH_TOMBSTONE) hash+=2;
+  if (table[slot].hash == hash) {
+    i = table[slot].index;
+    if (ht->pair[i].key_length == keylen && 0==memcmp(key, ht->pair[i].key, keylen))
+      goto delete;
+  }
+  if (table[slot].hash == HASH_EMPTY_MARKER)
+    return 0;
+
+  step = hash ^ (hash >> 16) ^ (hash >> 23);
+  step |= 1;
+  for(;;) {
+    slot = (slot + step) & mask;
+    if (table[slot].hash == hash) {
+      i = table[slot].index;
+      if (ht->pair[i].key_length == keylen && 0==memcmp(key, ht->pair[i].key, keylen))
+        goto delete;
+    }
+    if (table[slot].hash == HASH_EMPTY_MARKER)
+      return 0;
+  }
+ delete:
+  i = table[slot].index;
+  ht->pair[i].key        = 0;
+  ht->pair[i].value      = 0;
+  ht->pair[i].key_length = (size_t) ht->first_unused_pair;
+  ht->first_unused_pair = (int) i;
+  table[slot].hash = HASH_TOMBSTONE;
+  return 1;
+}
+
+static void stringhash_insert_ref(pphash *ht, table_entry e)
+{
+  uint32 mask = ht->table_mask;
+  uint32 hash = e.hash;
+  uint32 slot = hash & mask;
+  uint32 step;
+  table_entry *table = ht->table;
+  if (hash <= HASH_TOMBSTONE) hash+=2;
+
+  if (table[slot].hash != HASH_EMPTY_MARKER) {
+    step = hash ^ (hash >> 16) ^ (hash >> 23);
+    step |= 1;
+    for(;;) {
+      slot = (slot + step) & mask;
+      if (table[slot].hash == HASH_EMPTY_MARKER)
+        break;
+    }
+  }
+  table[slot] = e;
+}
+
+void stringhash_put(pphash *ht, char *key, size_t keylen, void *value)
+{
+  size_t n;
+  hashtable_pair p;
+  uint32 mask = ht->table_mask;
+  uint32 hash = preprocessor_hash(key, keylen);
+  uint32 slot = hash & mask;
+  uint32 step;
+  table_entry *table = ht->table;
+  if (hash <= HASH_TOMBSTONE) hash+=2;
+
+  if (ht->verify) { assert(stringhash_get(ht, key, keylen) == NULL); }
+  if (ht->verify) { char temp[1024]; memcpy(temp, key, keylen); temp[keylen] = 0; shput(ht->verify, temp, value); }
+
+  if (table[slot].hash > HASH_TOMBSTONE) {
+    step = hash ^ (hash >> 16) ^ (hash >> 23);
+    step |= 1;
+    for(;;) {
+      slot = (slot + step) & mask;
+      if (table[slot].hash <= HASH_TOMBSTONE)
+        break;
+    }
+  }
+  if (table[slot].hash == HASH_TOMBSTONE)
+    --ht->count_tombstones;
+
+  p.key = stb_arena_alloc_aligned(&ht->arena, keylen+1, 1);
+  p.key_length = keylen;
+  p.value = value;
+  memcpy(p.key, key, keylen);
+  p.key[keylen] = 0;
+  if (ht->first_unused_pair >= 0) {
+    // allocate existing slot from ht->pair
+    n = ht->first_unused_pair;
+    ht->first_unused_pair = (int) ht->pair[n].key_length;
+    ht->pair[n] = p;
+  } else {
+    // allocate new slot in ht->pair
+    arrput(ht->pair, p);
+    n = arrlen(ht->pair)-1;
+  }
+
+  table[slot].hash = hash;
+  table[slot].index = (uint32) n;
+  ++ht->count_in_use;
+
+  if (ht->count_in_use + ht->count_tombstones >= ht->doubling_threshold) {
+    table = ht->table;
+    uint32 size = ht->table_size;
+    uint32 i;
+    stringhash_init_table(ht, size*2);
+    ht->count_tombstones = 0;
+    for (i=0; i < size; ++i)
+      if (table[i].hash > HASH_TOMBSTONE)
+        stringhash_insert_ref(ht, table[i]);
+    free(table);
+  }
+}
+
+void stringhash_create(pphash *ht, int init_size) // must be power of two or 0
+{
+  memset(ht, 0, sizeof(pphash));
+  ht->first_unused_pair = -1;
+  stringhash_init_table(ht, init_size<4 ? 16 : init_size);
+  ht->verify = 0;
+  //sh_new_arena(ht->verify);
+}
+
+void stringhash_destroy(pphash *ht)
+{
+  stb_arena_free(&ht->arena);
+  arrfree(ht->pair);
+  free(ht->table);
+  if (ht->verify)
+    shfree(ht->verify);
+}
+
+char *get_string(int n)
+{
+  unsigned int bits = (unsigned int) n;
+  static char buffer[40];
+  int bottom_bits = (n & 3);
+  int rest = (n >> 2);
+  int len = ((rest*21) & 31), i;
+  if      (bottom_bits <  2) len = (len >> 2) + 3;
+  else if (bottom_bits == 2) len = (len >> 1) + 3;
+  else                       len =  len       + 3;
+  if ((n>>18) > 0) ++len;
+  len += 5;
+
+  bits *= xxprime3;
+  bits ^= bits >> 15;
+  bits = rotl(bits, 19);
+  for (i=0; i < len; ++i) {
+    buffer[i] = (bits & 63)+33;
+    bits = rotl(bits, 7);
+  }
+  // guarantee that first 2^24 strings are unique
+  buffer[0] = (n & 63)+33;
+  buffer[1] = ((n>>6)&63)+33;
+  buffer[2] = ((n>>12)&63)+33;
+  buffer[3] = ((n>>18)&63)+33;
+  buffer[len] = 0;
+  return buffer;
+}
+
+#ifdef HASHTEST
+void test_stringhash(void)
+{
+  int i,n;
+  pphash h;
+  for (n=6; n < 5000000; n *= 3) {
+    stringhash_create(&h, 16);
+    for (i=1; i < n; ++i) {
+      char *z = get_string(i);
+      stringhash_put(&h, z, strlen(z), (void*) (size_t) (i));
+    }
+    stringhash_destroy(&h);
+
+    stringhash_create(&h, 16);
+    for (i=1; i < n; ++i) {
+      char *z = get_string(i);
+      stringhash_put(&h, z, strlen(z), (void*) (size_t) (i));
+    }
+    for (i=1; i < n; ++i) {
+      char *z = get_string(i);
+      void *v = stringhash_get(&h, z, strlen(z));
+      assert(v == (void*) (size_t) i);
+    }
+    for (i=1; i < n; i += 2) {
+      char *z = get_string(i);
+      int r = stringhash_delete(&h, z, strlen(z));
+      assert(r == 1);
+    }
+    for (i=1; i < n; ++i) {
+      char *z = get_string(i);
+      void *v = stringhash_get(&h, z, strlen(z));
+      assert(v == ((i&1) ? 0 : (void*) (size_t) i));
+    }
+    for (i=1; i < n; ++i) {
+      char *z = get_string(i);
+      int r = stringhash_delete(&h, z, strlen(z));
+      assert(r == ((i&1) ? 0 : 1));
+    }
+    stringhash_destroy(&h);
+  }
+}
+#endif
+
+//////////////////////////////////////////////////////////////////////////////
+//
+//  SCANNERS
+//
+//    The following routines scan over characters looking for something.
+//
+//    It is expected that most per-character time will be spent in the
+//    scanners, thus we want high performance for many of them.
+//
+//    We have three classes of scanners, in order of speed:
+//
+//    1.  pp_char_flags scanners
+//    2.  state machine scanners
+//    3.  C-style scanners
+//
+//  pp_char_flags scanners
+//
+//    There are multiple character classifications which are assigned to
+//    individual bits of the 8-bit-per-character pp_char_flags[] array.
+//    A scanner can look for multiple bits at once, allowing this to be
+//    used for more than 8 kinds of scanning.
+//
+//  State Machine scanners
+//
+//    These scanners have a complex state transition diagram encoded
+//    in data tables which allows them to parse things like whitespace
+//    and comments simultaneously in a single fast loop.
+//
+//  C-style scanners
+//      
+//    A C-style scanner checks explicitly for all of the cases it cares
+//    about. It can do special multi-character processing, but is the
+//    slowest kind of scanner.
+
+//
+//  pp_char_flags scanners
+//
+
+static unsigned char pp_char_flags1[256];
+static unsigned char pp_char_flags2[256];
+
+#define CHAR1_IS_NUL                      1
+#define CHAR1_IS_NONNEWLINE_WHITESPACE    2
+#define CHAR1_IS_NEWLINE                  4
+#define CHAR1_IS_BALANCE_SCAN_ACTION      8  //  ( , ) EOF
+#define CHAR1_IS_PP_IDENTIFIER_FIRST     16
+#define CHAR1_IS_DIGIT                   32
+#define CHAR1_IS_QUOTE                   64
+#define CHAR1_IS_BACKSLASH              128
+
+#define CHAR2_IS_NUL                      1
+#define CHAR2_IS_SLASH                    2
+#define CHAR2_IS_NEWLINE                  4
+
+#define char_is_pp_identifier_first(x)     (pp_char_flags1[(uint8) (x)] &  CHAR1_IS_PP_IDENTIFIER_FIRST)
+#define char_is_pp_identifier(x)           (pp_char_flags1[(uint8) (x)] & (CHAR1_IS_PP_IDENTIFIER_FIRST | CHAR1_IS_DIGIT))
+#define char_is_nonnewline_whitespace(x)   (pp_char_flags1[(uint8) (x)] &  CHAR1_IS_NONNEWLINE_WHITESPACE)
+#define char_is_newline(x)                 (pp_char_flags1[(uint8) (x)] &  CHAR1_IS_NEWLINE)
+#define char_is_whitespace(x)              (pp_char_flags1[(uint8) (x)] & (CHAR1_IS_NONNEWLINE_WHITESPACE | CHAR1_IS_NEWLINE))
+#define char_is_nul(x)                     ((x) == 0)
+#define char_is_balance_scan_action(x)     (pp_char_flags1[(uint8) (x)] & (CHAR1_IS_BALANCE_SCAN_ACTION | CHAR1_IS_NUL))
+#define char_is_end_of_line(x)             (pp_char_flags1[(uint8) (x)] & (CHAR1_IS_NEWLINE | CHAR1_IS_NUL))
+#define char_is_whitespace_or_eof(x)       (pp_char_flags1[(uint8) (x)] & (CHAR1_IS_NONNEWLINE_WHITESPACE | CHAR1_IS_NEWLINE | CHAR1_IS_NUL))
+#define char_is_special_string_contents(x) (pp_char_flags1[(uint8) (x)] & (CHAR1_IS_QUOTE | CHAR1_IS_BACKSLASH | CHAR1_IS_NUL))
+#define char_is_quote(x)                   (pp_char_flags1[(uint8) (x)] & CHAR1_IS_QUOTE)
+#define char_is_argument_parse_action(x)   (pp_char_flags1[(uint8) (x)] & (CHAR1_IS_BALANCE_SCAN_ACTION | CHAR1_IS_QUOTE | CHAR1_IS_BACKSLASH | CHAR1_IS_NUL | CHAR1_IS_NEWLINE))
+#define char_is_copy_line_exception(x)     (pp_char_flags2[(uint8) (x)] & (CHAR2_IS_SLASH | CHAR2_IS_NEWLINE | CHAR2_IS_NUL))
+
+#ifdef TREAT_LF_CR_AS_TWO_NEWLINES
+  // returns true if CR LF
+  #define is_twochar_newline_given_first_is_newline(a,b)  ((a) == '\r' && (b) == '\n')
+  //#define is_twochar_newline_given_first_is_newline(a,b)  ((a) == '\r' & (b) == '\n') // potentially faster but probably slower
+#else
+  // returns true if CR LF or LF CR
+  #define is_twochar_newline_given_first_is_newline(a,b)  ((a) + (b) == '\r' + '\n')
+#endif
+
+#define newline_char_count(newline_char, following_char)  (is_twochar_newline_given_first_is_newline(newline_char,following_char) ? 2 : 1)
+
+#define is_backslash_newline(p)  ((p)[0] == '\\' && char_is_newline((p)[1]))
+
+// this is used when parsing a backslash-eliminated directive
+static char *preprocessor_skip_whitespace_simple(char *s)
+{
+  while (char_is_nonnewline_whitespace(*s))
+    ++s;
+  return s;
+}
+
+// This is used when trimming macro arguments of leading whitespace
+static char *preprocessor_skip_whitespace_fast(char *s)
+{
+  // same as preprocessor_skip_whitespace, but doesn't need to count newlines
+  for(;;) {
+    while (char_is_whitespace(*s))
+      ++s;
+
+    // handle backslash continuation
+    if (s[0] == '\\' && char_is_newline(s[1]))
+      s += 2;
+    else
+      return s;
+  }
+}
+
+// This is used when trimming macro arguments of trailing whitespace
+// @OPTIMIZE: turn this into a smart state-machine scanner?
+static char *preprocessor_skip_whitespace_reverse(char *s)
+{
+  for(;;) {
+    while (char_is_whitespace(s[-1]))
+      --s;
+
+    // handle backslash continuation
+    if (s[-1] == '\\' && char_is_newline(s[1])) {
+      --s;
+    } else
+      return s;
+  }
+}
+
+static char *scan_pp_identifier(char *p)
+{
+  assert(char_is_pp_identifier_first(*p));
+  ++p;
+  while (char_is_pp_identifier(*p))
+    ++p;
+  return p;
+}
+
+// after running fast whitespace & comment scanner, we lose sync with start-of-line,
+// which will break column numbers in compiler, so this helper function recovers
+static void rewind_to_start_of_line(parse_state *cs)
+{
+  char *q = cs->src;
+  char *p = cs->src + cs->src_offset;
+
+  // stop at end of line or non-whitespace (which actually should only be able to be end of a multiline comment, i.e. '/')
+  while (p > q && char_is_nonnewline_whitespace(p[-1]))
+    --p;
+
+  cs->src_offset = p - q;
+  cs->in_leading_whitespace = 1;
+}
+
+//
+//  State Machine scanners
+//
+
+enum {
+  PP_CHAR_CLASS_default,
+  PP_CHAR_CLASS_whitespace,
+  PP_CHAR_CLASS_cr,
+  PP_CHAR_CLASS_lf,
+  PP_CHAR_CLASS_period,
+  PP_CHAR_CLASS_digit,
+  PP_CHAR_CLASS_idtype, // letter, _, or $
+  PP_CHAR_CLASS_idtype_e_E,
+  PP_CHAR_CLASS_apostrophe,
+  PP_CHAR_CLASS_quotation,
+
+  PP_CHAR_CLASS_slash,
+  PP_CHAR_CLASS_backslash,
+  PP_CHAR_CLASS_star,
+  PP_CHAR_CLASS_hash,
+  PP_CHAR_CLASS_plusminus,
+
+  PP_CHAR_CLASS_eof,
+
+  PP_CHAR_CLASS_count,
+};
+
+enum {
+  PP_STATE_ready,
+  PP_STATE_in_leading_whitespace,
+  PP_STATE_saw_slash,
+  PP_STATE_comment_singleline,
+  PP_STATE_comment_multiline,
+  PP_STATE_comment_multiline_seen_star,
+  PP_STATE_comment_multiline_seen_cr,
+  PP_STATE_comment_multiline_seen_lf,
+  PP_STATE_in_string_literal,
+  PP_STATE_in_string_literal_saw_backslash,
+  PP_STATE_in_char_literal,
+  PP_STATE_in_char_literal_saw_backslash,
+  PP_STATE_saw_period,
+  PP_STATE_in_number,
+  PP_STATE_in_number_saw_e,
+  PP_STATE_saw_cr,
+  PP_STATE_saw_lf,
+
+  PP_STATE_active_count,
+
+  PP_STATE_saw_backslash,
+
+  PP_STATE_eof,
+  PP_STATE_saw_leading_hash,
+  PP_STATE_saw_hash,
+  PP_STATE_saw_identifier_start,
+
+  PP_STATE_error_unterminated_multiline_comment,
+
+  PP_STATE_count
+};
+
+// here are possible exit states, and their meanings:
+//
+//  PP_STATE_saw_cr,                               // if we need to emit #line
+//  PP_STATE_saw_lf,                               // if we need to emit #line
+//  PP_STATE_saw_backslash,                        // backslash at end of line or syntax error, but not in preprocessor directive such as macro definition
+//  PP_STATE_eof,                                  // end of file
+//  PP_STATE_saw_leading_hash,                     // preprocessor directive
+//  PP_STATE_saw_hash,                             // syntax error    // @TODO: get rid of this case and let the compiler detect and report the error
+//  PP_STATE_saw_identifier_start,                 // any identifier 
+//  PP_STATE_error_unterminated_multiline_comment  // syntax error
+
+enum {
+  PP_WHITESCAN_STATE_whitespace,
+  PP_WHITESCAN_STATE_saw_cr,
+  PP_WHITESCAN_STATE_saw_lf,
+  PP_WHITESCAN_STATE_saw_slash,
+  PP_WHITESCAN_STATE_comment_singleline,
+  PP_WHITESCAN_STATE_comment_multiline,
+  PP_WHITESCAN_STATE_comment_multiline_seen_cr,
+  PP_WHITESCAN_STATE_comment_multiline_seen_lf,
+  PP_WHITESCAN_STATE_comment_multiline_seen_star,
+
+  PP_WHITESCAN_STATE_active_count,
+  PP_WHITESCAN_STATE_stop_seen_onechar,
+  PP_WHITESCAN_STATE_stop_seen_twochars,
+  PP_WHITESCAN_STATE_stop_eof,
+  PP_WHITESCAN_STATE_error_unterminated_multiline_comment,
+};
+
+
+static unsigned char pp_char_class[256];
+static unsigned char pp_scan_char_class[256];      // same as above, but identifiers are treated as default
+static unsigned char pp_transition_table[PP_STATE_active_count][16]; // normal transition table
+static unsigned char pp_transition_initscan[PP_STATE_active_count][16]; // special transition table while scanning for initial #ifndef
+static unsigned char pp_state_is_end_of_line[PP_STATE_count];
+static unsigned char pp_initscan_state_is_end_of_line[PP_STATE_count];
+
+static void output_line_directive(parse_state *cs);
+
+// To implement fast include-guard skipping, we want to detect #ifndef at start of file.
+// So we start by scanning whitespace and comments until we hit the first non-whitespace
+// characters in the file. We don't handle backslash concatenation, so no fast include-guard if so.
+static int scan_whitespace_and_comments(parse_state *cs)
+{
+  char *p   = cs->src + cs->src_offset;
+  int state = PP_WHITESCAN_STATE_whitespace;
+  int line_count = cs->src_line_number;
+
+  do {
+    uint8 ch = (uint8) *p++;
+    uint8 ch_class = pp_char_class[ch];
+    state = pp_transition_initscan[state][ch_class];
+    line_count += pp_initscan_state_is_end_of_line[state];
+  } while (state < PP_WHITESCAN_STATE_active_count);
+
+  cs->src_line_number = line_count;
+
+  if (state == PP_WHITESCAN_STATE_stop_seen_twochars) {
+    // we've processed two characters, and incremented past both of them
+    p -= 2;
+  } else {
+    p -= 1;
+  }
+
+  cs->src_offset = p - cs->src;
+
+  return state;
+}
+
+// Following routine is the workhorse; this copies from a file's text buffer
+// into the output buffer until it finds a directive, a preprocessor Identifier
+// (which might be a macro name), or the end of the file.
+//
+// It also performs backslash-continuation elimination.
+
+static int copy_to_action_point(parse_state *cs)
+{
+  int state;
+  char *p   = cs->src + cs->src_offset;
+  char *out;
+  char *q;
+
+  out = cs->dest;
+
+  state = cs->in_leading_whitespace ? PP_STATE_in_leading_whitespace : PP_STATE_ready;
+  MDBG();
+
+  if (cs->fast_dest)
+    q = cs->fast_dest;
+  else
+    q = NULL;
+
+  for(;;) {
+    int prev_state;
+
+    if (cs->fast_dest == NULL) {
+      // expand output so there's room to copy entire file unmodified
+      arrsetcap(out, arrlen(out) + (cs->src_length - cs->src_offset) + 1);
+      q = out + arrlen(out);
+    }
+
+    { // make lifetime of line_count unambiguous
+      const int state_limit = cs->state_limit;
+      int line_count = 0;
+
+      do {
+        uint8 ch = (uint8) *p++;
+        uint8 ch_class;
+        assert(cs->fast_dest != NULL || q < out + arrcap(out));
+        *q++ = ch;
+        ch_class = pp_char_class[ch];
+        prev_state = state;
+        state = pp_transition_table[state][ch_class];
+        line_count += pp_state_is_end_of_line[state];
+      } while (state < state_limit);
+
+      cs->src_line_number += line_count;
+      cs->dest_line_number += line_count;
+    }
+
+    // The cases that are not syntax errors and not once per file:
+    //    need to emit #line
+    //    preprocessor directive
+    //    any identifier
+    //    backlash at end of non-preprocessor directive
+    //
+    // of those, presumably most common case is going to be an identifier
+    // (which we have to check to see if it's a macro). So we should handle
+    // that case first. We can use the filters that we use to detect identifiers
+    // that aren't macros to ALSO detect not any of the other cases, i.e. naively if
+    // the filter returns true it's because it MUST be something other than
+    // an identifier that's not a macro, but we can also make it ONLY return true
+    // if the input was an identifier and apply it BEFORE checking we're in the
+    // identifier case
+
+    if (state > PP_STATE_saw_backslash)
+      break;
+
+    if (state == PP_STATE_saw_backslash) {
+      assert(p[-1] == '\\');
+      if (!char_is_newline(*p)) {
+        // if there's not a newline, just treat backslash as normal character and continue
+        state = PP_STATE_ready;
+      } else {
+        // join the lines together
+        --q;                                   // undo copying the backlash
+        p += newline_char_count(p[0], p[1]);   // skip past CR, LF, CR LF, LF CR
+        cs->src_line_number += 1;              // we've seen it in the source but haven't seen it in the dest
+        cs->state_limit = PP_STATE_saw_cr;     // line_number mismatches, so need to output match newlines at first opportunity, namely next time we see a newline not in a comment
+        state = prev_state;                    // reset parse state to before the backslash
+        if (cs->fast_dest == NULL)
+          arrsetlen(out, (size_t) (q-out));
+      }
+    } else {
+      // after the above backslash code runs, this code
+      // is run on the next CR or LF
+
+      assert(p[-1] == '\n' || p[-1] == '\r');
+
+      // finish copying a CR LF or LF CR pair
+      if (is_twochar_newline_given_first_is_newline(p[-1],p[0]))
+        *q++ = *p++;
+
+      // reset state
+      cs->state_limit = PP_STATE_active_count;
+
+      // output newlines to the output buffer until synchronized again
+      if (cs->fast_dest) {
+        int n = cs->src_line_number - cs->dest_line_number;
+        while (n > 0) {
+          *q++ = '\n';
+          --n;
+        }
+        cs->dest_line_number = cs->src_line_number;
+      } else {
+        arrsetlen(out, (size_t) (q - out));
+        while (cs->dest_line_number < cs->src_line_number) {
+          arrput(out, '\n');
+          ++cs->dest_line_number;
+        }
+      }
+    }
+  }
+
+  // set output to the amount we actually copied... don't count the copy of the last character
+  if (cs->fast_dest) {
+    cs->fast_dest = q-1;
+  } else {
+    arrsetlen(out, (size_t) (q-1 - out));
+    cs->dest = out;
+  }
+  cs->src_offset = (size_t) (p-1 - cs->src);
+  cs->in_leading_whitespace = (state == PP_STATE_saw_leading_hash);
+
+  return state;
+}
+
+// after performing macro expansion, we rescan the output for further
+// macros. this doesn't need to do newline processing or handle directives.
+// we use the same states as above and just should never hit the other ones.
+// because we're not doing backslash-continuations, the inner loop can
+// be simpler (no tracking of previous state)...
+//
+// (we do need to parse and // skip newlines, since they can be copied from
+// macro arguments, but we don't need to COUNT them)
+static int copy_to_action_point_macro_expansion(parse_state *cs)
+{
+  int state;
+  char *p   = cs->src + cs->src_offset;
+  char *out = cs->dest;
+  char *q;
+
+  if (cs->fast_dest)
+    q = cs->fast_dest;
+  else {
+    // expand output so copying entire file unmodified
+    arrsetcap(out, arrlen(out) + (cs->src_length - cs->src_offset) + 1);
+    q = out + arrlen(out);
+  }
+
+  state = PP_STATE_ready;
+
+  do {
+    uint8 ch = (uint8) *p++;
+    uint8 ch_class;
+    *q++ = ch;
+    ch_class = pp_char_class[ch];
+    state = pp_transition_table[state][ch_class];
+  } while (state < PP_STATE_active_count);
+
+  cs->src_offset = (size_t) (p-1 - cs->src);
+
+  if (cs->fast_dest)
+    cs->fast_dest = q-1;
+  else {
+    // set output to the amount we actually copied... don't count the copy of the last character
+    arrsetlen(out, (size_t) (q-1 - out));
+    cs->dest = out;
+  }
+
+  return state;
+}
+
+// when in a conditionally-disabled section of code, we want to
+// scan as fast as possible to the next directive; need a state-machine
+// to handle comments efficiently
+static char *scan_to_directive(char *p, int *p_line_number)
+{
+  int line_number = *p_line_number;
+  int state = PP_STATE_ready, prev_state;
+  for (;;) {
+    do {
+      uint8 ch = *p++;
+      uint8 ch_class = pp_scan_char_class[ch]; // don't stop at identifiers
+      prev_state = state;
+      state = pp_transition_table[state][ch_class];
+      line_number += pp_state_is_end_of_line[state];
+    } while (state < PP_STATE_active_count);
+  
+    if (state != PP_STATE_saw_backslash) {
+      *p_line_number = line_number;
+      return p-1;
+    }
+
+    assert(p[-1] == '\\');
+
+    // if character after \ is CR/LF, join the lines together
+    if (*p == '\n' || *p == '\r') {
+      p += newline_char_count(p[0], p[1]);       // skip past CR, LF, CR LF, LF CR
+      ++line_number;
+      state = prev_state;                        // reset parse state to before the backslash
+    } else
+      state = PP_STATE_ready;
+  }
+}
+
+//
+//  C-style scanners
+//
+
+// this is used when parsing macro arguments... so @TODO comments could be present
+static char *preprocessor_skip_whitespace(char *s, int *newline_count)
+{
+  for(;;) {
+    while (char_is_whitespace(*s)) {
+      while (char_is_nonnewline_whitespace(*s))
+        ++s;
+
+      while (char_is_newline(*s)) {
+        s += newline_char_count(s[0],s[1]);
+        ++*newline_count;
+      }
+    }
+
+    // handle backslash continuation
+    if (s[0] == '\\' && char_is_newline(s[1])) {
+      ++s;
+      s += newline_char_count(s[0], s[1]);
+      ++*newline_count;
+    } else
+      return s;
+  }
+}
+
+// this is used when parsing directives, so newlines terminate it
+// comments are allowed but rare.
+static char *preprocessor_skip_whitespace_in_directive(char *s, int *newline_count)
+{
+  for(;;) {
+
+    // loop over whitespace 
+    while (char_is_nonnewline_whitespace(*s))
+      ++s;
+
+    if (*s == '/') {
+      // skip comments
+      if (s[1] == '/') goto scan_to_end_of_line;
+      if (s[1] != '*')
+        break;
+      // parse multiline comment
+      s += 2;
+      for (;;) {
+        if (s[0] == 0)
+          break; // @TODO: error, end of file in multiline comment
+        if (s[0] == '*' && s[1] == '/') {
+          s += 2;
+          break;
+        }
+        ++s;
+      }
+    } else if (*s == '\\' && char_is_newline(s[1])) {
+      ++s;
+      s += newline_char_count(s[0], s[1]);
+      ++newline_count;
+    } else
+      return s;
+  }
+ scan_to_end_of_line:
+  while (!char_is_end_of_line(*s))
+    ++s;
+  return s;
+}
+
+static char *copy_line_without_comments(char *buffer, size_t buffer_size, char *p, int *line_count, char **endp)
+{
+  char *out = buffer, *end = buffer + buffer_size;
+  while (char_is_nonnewline_whitespace(*p))
+    ++p;
+
+  for(;;) {
+    // fast copy into buffer
+    while (out < end && !char_is_copy_line_exception(*p))
+      *out++ = *p++;
+
+    // handle end of line & backslash
+    if (char_is_end_of_line(*p)) {
+      if (p[0] == 0 || p[-1] != '\\') {
+        *out = 0;
+        *endp = p;
+        return buffer;
+      }
+      // process backslash escape by ungetting the backslash and skipping the newline
+      --out;
+      p += newline_char_count(p[0],p[1]);
+      ++*line_count;
+    } else if (p[0] == '/') {
+      if (p[1] == '/') {
+        while (!char_is_end_of_line(*p))
+          ++p;
+        *out = 0;
+        *endp = p;
+        return buffer;
+      } else if (p[1] == '*') {
+        p += 2;
+        while (!(p[0] == '*' && p[1] == '/')) {
+          if (char_is_end_of_line(*p)) {
+            if (*p == 0)
+              return buffer;
+            p += newline_char_count(p[0], p[1]);
+            ++*line_count;
+          } else
+            ++p;
+        }
+        p += 2;
+      } else {
+        *out++ = *p++;
+      }
+    } else {
+      assert(out >= end);
+    }
+
+    if (out >= end)
+      break;
+  }
+
+  // if we reach here, the copy above overflowed the buffer
+  {
+    size_t length = out - buffer;
+    out = NULL; // stb_ds array
+    arrsetlen(out, length);
+    memcpy(out, buffer, length);
+
+    // this case is most likely to be a long #define
+    for(;;) {
+      size_t count=0;
+      while (!char_is_copy_line_exception(*p)) {
+        ++count;
+        ++p;
+      }
+      arrsetlen(out, arrlen(out) + count);
+      memcpy(out + arrlen(out) - count, p - count, count);
+
+      // handle end of line & backslash
+      if (char_is_end_of_line(*p)) {
+        if (p[0] == 0 || p[-1] != '\\') {
+          arrput(out, 0);
+          *endp = p;
+          return out;
+        }
+        // process backslash escape by ungetting the backslash and skipping the newline
+        arrsetlen(out, arrlen(out)-1);
+        p += newline_char_count(p[0],p[1]);
+        ++*line_count;
+      } else if (p[0] == '/') {
+        if (p[1] == '/') {
+          while (!char_is_end_of_line(*p))
+            ++p;
+          arrput(out,0);
+          *endp = p;
+          return out;
+        } else if (p[1] == '*') {
+          p += 2;
+          while (!(p[0] == '*' && p[1] == '/')) {
+            if (char_is_end_of_line(*p)) {
+              if (*p == 0)
+                return buffer; // @TODO: unterminated multline comment needs to report error
+              p += newline_char_count(p[0], p[1]);
+              ++*line_count;
+            } else
+              ++p;
+          }
+          p += 2;
+        } else {
+          arrput(out, *p++);
+        }
+      } else {
+        assert(0);
+      }
+    }
+  }
+}
+
+// copy until a newline, but with backslash-continuation elimination
+static char* convert_backslash_continuation(char *src, char *no_slashes_before_this_point, char **update_src, int *line_number)
+{
+  char *copy = NULL;
+  char *p = src;
+
+  (void) sizeof(no_slashes_before_this_point); // @OPTIMIZE faster copy up to this point to avoid rescanning
+
+  // Implementation strategies:
+  //    1. scan entire line to find length; alloc; copy; this will do 2n over the input
+  //    2. use dynamic array to assemble; this will do n over the input and O(n) copies of the output during realloc
+  // But 2 can avoid realloc if the initial size is large enough
+  //
+  // @OPTIMIZE: faster if we don't use stb_ds but use explicit buffer/capacity/length,
+  // so we can keep them in registers.
+
+  arrsetcap(copy, 1024);
+  for(;;) {
+    while (!char_is_end_of_line(*p) && *p != '/') {
+      arrput(copy, *p);
+      ++p;
+    }
+    if (*p == 0)
+      break;
+    if (*p == '/') {
+      if (p[1] == '/') {
+        while (!char_is_end_of_line(*p))
+          ++p;
+        break;
+      } else if (p[1] == '*') {
+        p += 2;
+        while (*p != 0 && (p[0] != '*' || p[1] != '/'))
+          ++p;
+        p += 2;
+        arrput(copy, ' '); // avoid token pasting
+        continue;
+      } else {
+        arrput(copy, *p++);
+        continue;
+      }
+    }
+    assert(char_is_newline(*p));
+    // if newline is preceded by backslash, keep going
+    if (p[-1] != '\\')
+      break;
+    p += newline_char_count(p[0],p[1]);
+    ++*line_number;
+    arrsetlen(copy, arrlen(copy)-1); // undo the output of the backslash
+  }
+
+  arrput(copy, 0);
+
+  if (update_src)
+    *update_src = p;
+
+  return copy;
+}
+
+static char *maybe_convert_backslash_continuation(char *src, char **update_src, char **alloc_buffer, ptrdiff_t *length, int *line_number)
+{
+  char *p = src;
+  while (!char_is_end_of_line(*p) && *p != '/')
+    ++p;
+  if (*p == 0) {
+    *length = p - src;
+    *update_src = p;
+    return src;
+  }
+  if ((char_is_newline(*p) && p[-1] == '\\') || *p == '/') {
+    char *alloc = convert_backslash_continuation(src, p-1, update_src, line_number);
+    *alloc_buffer = alloc;
+    *length = arrlen(alloc)-1; // don't count NUL
+    return alloc;
+  } else {
+    assert(char_is_newline(*p));
+    *length = p - src;
+    *update_src = p;
+    return src;
+  }
+}
+
+//////////////////////////////////////////////////////////////////////////////
+//
+//  MAIN PROCESSING LOOP
+//
+
+enum
+{
+  IN_MACRO_no,
+  IN_MACRO_yes,
+  IN_MACRO_directive,  // e.g. expanding #include or #pragma
+  IN_MACRO_if_condition,  // if we're expanding '#if' so we need to do special processing for defined()
+  IN_MACRO_include_guard_scan,
+};
+
+enum
+{
+  CONDITIONAL_none,
+  CONDITIONAL_skip_current_block,
+  CONDITIONAL_skip_all_blocks,
+};
+
+typedef struct
+{
+  int conditionally_disable;
+  int conditional_compilation_nesting_level;
+} conditional_state;
+
+static void process_directive(parse_state *cs, conditional_state *cons);
+static void maybe_expand_macro(parse_state *cs, macro_definition *pending, int in_macro_expansion, char *in_macro, macro_definition **md);
+
+static void preprocess_string(parse_state *cs, int in_macro_expansion, char *in_macro, macro_definition **final_function_macro)
+{
+  pp_context *c = cs->context;
+  conditional_state cons;
+  cons.conditionally_disable                 = CONDITIONAL_none;
+  cons.conditional_compilation_nesting_level = 0;
+  cs->include_guard_candidate_macro_name = 0;
+
+  if (in_macro_expansion == IN_MACRO_include_guard_scan) {
+    char *p;
+    int action;
+
+    in_macro_expansion = IN_MACRO_no;
+    c->include_guard_detect = 1;
+    
+    // check for leading #ifndef
+    action = scan_whitespace_and_comments(cs);
+    if (action != PP_WHITESCAN_STATE_stop_seen_onechar || cs->src[cs->src_offset] != '#')
+      goto handle_action;
+
+    p = preprocessor_skip_whitespace_simple(cs->src+cs->src_offset+1); // skip past '#'
+    if (0 != memcmp(p, "ifndef", 6))
+      goto handle_action;
+      
+    process_directive(cs, &cons);
+    if (c->last_ifndef == 0)
+      goto done;
+
+    // check for following #define
+    action = scan_whitespace_and_comments(cs);
+    if (action != PP_WHITESCAN_STATE_stop_seen_onechar || cs->src[cs->src_offset] != '#')
+      goto handle_action;
+    p = preprocessor_skip_whitespace_simple(cs->src+cs->src_offset+1); // skip past '#'
+    if (0 != memcmp(p, "define", 6)) {
+     handle_action:
+      // rewind to good point to continue scanning with regular scanner, to keep code simple
+      switch (action) {
+        case PP_WHITESCAN_STATE_stop_seen_onechar : rewind_to_start_of_line(cs); break;
+        case PP_WHITESCAN_STATE_stop_seen_twochars: rewind_to_start_of_line(cs); break;
+        case PP_WHITESCAN_STATE_stop_eof          : break;
+        case PP_WHITESCAN_STATE_error_unterminated_multiline_comment: goto unterminated_multline_comment; break;
+      }
+     done:
+      ;
+    } else {
+      process_directive(cs, &cons);
+      if (c->last_macro_definition && 0==strcmp(c->last_macro_definition->symbol_name, c->last_ifndef)) {
+        cs->include_guard_candidate_macro_name = c->last_ifndef;
+      }
+    }
+    c->include_guard_detect = 0;
+  }
+
+
+  for (;;) {
+    //char *start = cs->src + cs->src_offset; // use this to rediscover the start of an unterminated /* comment
+    int action;
+
+    // only output line directives at start of line
+    if (cs->in_leading_whitespace && in_macro_expansion == IN_MACRO_no) {
+      if (cs->src_line_number != cs->dest_line_number || cs->last_output_filename != cs->filename) {
+        if (cs->last_output_filename == cs->filename && cs->dest_line_number < cs->src_line_number && cs->src_line_number < cs->dest_line_number+12) {
+          int n = cs->src_line_number - cs->dest_line_number;
+          if (cs->fast_dest) {
+            while (n > 0) {
+              *cs->fast_dest++ = '\n';
+              --n;
+            }
+            cs->dest_line_number = cs->src_line_number;
+          } else {
+            while (cs->dest_line_number < cs->src_line_number) {
+              arrput(cs->dest, '\n');
+              ++cs->dest_line_number;
+            }
+          }
+        } else
+          output_line_directive(cs);
+      }
+    }
+
+    action = in_macro_expansion > IN_MACRO_no
+                     ? copy_to_action_point_macro_expansion(cs)
+                     : copy_to_action_point(cs);
+
+    assert(cons.conditional_compilation_nesting_level >= 0);
+
+    switch (action) {
+      case PP_STATE_saw_leading_hash: {
+        MDBG();
+        process_directive(cs, &cons);
+        if (c->stop)
+          return;
+        cs->in_leading_whitespace = 1;
+        MDBG();
+        break;
+      }
+
+      case PP_STATE_saw_identifier_start: {
+        macro_definition *md = 0;
+        maybe_expand_macro(cs, NULL, in_macro_expansion, in_macro, &md);
+       iterate_final_macro:
+        if (c->stop)
+          return;
+        cs->in_leading_whitespace = 0;
+        if (md != NULL) {
+          // we saw a function-like macro name at the end of the macro expansion above, but no (.
+          // so we now check to see whether it's followed by (. clang looks inside following macros
+          // for this, but VC6 doesn't, and it would require a rearchitecture to look inside the macros
+          // (and probably an efficiency hit), so we'll do it the VC6 way.
+          char *s = cs->src + cs->src_offset;
+          // @TODO skip comments and newlines
+          while (char_is_nonnewline_whitespace(*s))
+            ++s;
+          cs->src_offset = s - cs->src;
+          if (*s == '(') {
+            macro_definition *temp = md;
+            md = NULL;
+            maybe_expand_macro(cs, temp, in_macro_expansion, in_macro, &md);
+            goto iterate_final_macro;
+          } else {
+            if (*s == 0 && final_function_macro != NULL) {
+              *final_function_macro = md;
+              return;
+            }
+          }
+          // the final macro wasn't copied into the output destination yet, so do it now
+          {
+            size_t len = strlen(md->symbol_name);
+            if (cs->fast_dest) {
+              memcpy(cs->fast_dest, md->symbol_name, len);
+              cs->fast_dest += len;
+            } else {
+              size_t off = arraddnindex(cs->dest, len);
+              memcpy(cs->dest+off, md->symbol_name, len);
+            }
+          }
+        }
+        break;
+      }
+
+      case PP_STATE_saw_hash: {
+        if (do_error_code(cs, PP_RESULT_directive_not_at_start_of_line, "Preprocessor directive must appear at beginning of line."))
+          return;
+        else {
+          // if we're supposed to keep going, write it to the output; maybe you're using this with a language that allows # as a token
+          arrput(cs->dest, '#');
+          ++cs->src_offset;
+        }
+        break;
+      }
+
+      case PP_STATE_eof: {
+        if (cons.conditional_compilation_nesting_level) {
+          int i, stop=0;
+          for (i=(int) arrlen(c->ifdef_stack)-1; i >= 0; --i) {
+            if (do_error(cs, "No #endif found for previous #%s", directive_hash[c->ifdef_stack[i].type].name))
+              stop = 1;
+            error_supplement(cs, PP_RESULT_ERROR, c->ifdef_stack[i].line, "Location of #%s", directive_hash[c->ifdef_stack[i].type].name);
+          }
+          if (stop)
+            return;
+        }
+        return;
+      }
+
+      unterminated_multline_comment:
+      case PP_STATE_error_unterminated_multiline_comment: {
+        // @TODO: locate the comment's starting location, which requires fully rescanning
+        // from 'start'
+        do_error(cs, "Unclosed '/*' comment.");
+        return;
+      }
+
+      default:
+        assert(0);
+        break;
+    }
+
+    MDBG();
+  }
+}
+
+//////////////////////////////////////////////////////////////////////////////
+//
+//  MACRO PROCESSING
+//
+
+#define MAX_MACRO_PARAMETER_NAME_LENGTH   256
+
+static void update_macro_filter(pp_context *c, char *identifier)
+{
+  assert(identifier[0] != 0);
+  if (identifier[1] == 0 || identifier[2] == 0)
+    c->macro_name_filter[(uint8)identifier[0]][(uint8)identifier[1]] = 1;
+}
+
+static macro_definition *create_macro_definition(parse_state *ps, char *def)
+{
+  // keep track of the state so we can share the diagnostic code
+  enum {
+    STATE_pre_name,
+    STATE_object_like,
+    STATE_function_like
+  } state = STATE_pre_name;
+
+  pp_context *c = ps->context;
+  char identifier_buffer[MAX_MACRO_PARAMETER_NAME_LENGTH+1];
+  macro_definition *m = stb_arena_alloc(&c->macro_arena, sizeof(*m));
+  char *start = preprocessor_skip_whitespace_simple(def);
+  char *p;
+
+  m->symbol_name = NULL;
+  m->disabled    = 0;
+  m->predefined  = 0;
+
+  // parse identifier
+  p = start;
+  if (!char_is_pp_identifier_first(*p))
+    goto invalid_char_error;
+  while (char_is_pp_identifier(*p))
+    ++p;
+
+  m->symbol_name = stb_arena_alloc_string_length(&c->macro_arena, start, p-start);
+  m->symbol_name_length = (int) (p-start);
+
+  if (*p != '(') {
+    state = STATE_object_like;
+
+    // object-like definition
+    if (!char_is_whitespace_or_eof(*p))
+      goto invalid_char_error;
+
+    // find the start of the definition
+    p = preprocessor_skip_whitespace_simple(p);
+
+    start = p;
+
+    while (!char_is_end_of_line(*p))
+      ++p;
+
+    // eliminate trailing whitespace
+    if (p > start) {
+      while (char_is_whitespace(p[-1]))
+        --p;
+      assert(p > start);
+    }
+
+    m->num_parameters          = MACRO_NUM_PARAMETERS_no_parentheses;
+    m->simple_expansion        = stb_arena_alloc_string_length(&c->macro_arena, start, p-start);
+    m->simple_expansion_length = (int) (p-start);
+  } else {
+    // function-like definition
+    pphash parameters;
+    int token_paste_before_text = 0;
+    state = STATE_function_like;
+    m->num_parameters = 0;
+    m->is_variadic = 0;
+    assert(*p == '(');
+    // in the () case, we can do simple non-function-like expansion, but to share
+    // code, we use this same path, but that means in the () case we still need to
+    // allocate the parameters dictionary so it can return -1 as the default
+
+    stringhash_create(&parameters, 64); // oversized so hash table is faster
+
+    p = preprocessor_skip_whitespace_simple(p+1); // skip paren then following whitespace
+    // accept either ) or C Identifier or ...
+    if (*p == ')') {
+      ; // do nothing, we have defined as having 0 parameter
+    } else {
+      if (*p == '.') {
+        if (p[1] == '.' && p[2] == '.') {
+          p = preprocessor_skip_whitespace_simple(p+3);
+          if (*p != ')')
+            goto invalid_char_error;
+          m->is_variadic = 1;
+          stringhash_put(&parameters, "__VA_ARGS__", 11, (void*) (size_t) 1);
+        } else
+          goto invalid_char_error;
+      } else if (char_is_pp_identifier_first(*p)) {
+        // parse the parameters
+        for(;;) {
+          start = p = preprocessor_skip_whitespace_simple(p);
+          if (!char_is_pp_identifier_first(*p)) {
+            if (p[0] == '.' && p[1] == '.' && p[2] == '.') {
+              p = preprocessor_skip_whitespace_simple(p+3);
+              if (*p != ')')
+                goto invalid_char_error;
+              m->is_variadic = 1;
+              stringhash_put(&parameters, "__VA_ARGS__", 11, (void*) (size_t) (m->num_parameters+1));
+              break;
+            } else {
+              goto invalid_char_error;
+            }
+          } else {
+            size_t length;
+            ++p;
+            while (char_is_pp_identifier(*p))
+              ++p;
+            length = p-start;
+
+            if (stringhash_get(&parameters, start, length) != 0) {
+              if (length > sizeof(identifier_buffer)-1)
+                length = sizeof(identifier_buffer)-1;
+              memcpy(identifier_buffer, start, length);
+              identifier_buffer[length] = 0;
+              do_error(ps, "Macro '%s' contained the parameter '%s' more than once", m->symbol_name, identifier_buffer);
+              goto error_exit;
+            }
+            stringhash_put(&parameters, start, length, (void*) (size_t) ++m->num_parameters);
+            p = preprocessor_skip_whitespace_simple(p);
+            if (*p == ')')
+              break; // dont with parameter loop
+            if (*p == ',') {
+              ++p;
+              continue;
+            }
+            goto invalid_char_error;
+          }
+        }
+      } else if (*p == 0) {
+        goto invalid_char_error;
+      }
+    }
+
+    assert(*p == ')');
+    p = preprocessor_skip_whitespace_simple(p+1);
+
+    // now the remainder of p is the substitution string. we need to parse it
+    // for identifiers that are in the parameters table
+    m->expansion = NULL;
+    token_paste_before_text = 0; // @TODO: this used to be 1 to make sure we didn't generate a leading space, but this was wrong with reparsing. but may be problematic with nesting now?
+    for(;;) {
+      int found_parameter = -1;
+      int require_identifier = 0;
+      int stringify = 0;
+      int token_paste_after_text = 0;
+      expansion_part e;
+      char *end;
+      start = p;
+      // we need to parse to an identifier or '#', and we need to
+      // skip comments and string literals. we maybe could speed
+      // this up with a scanner, but we probably don't spend that
+      // much time here anyway
+      for(;;) {
+        if (char_is_quote(*p)) {
+          char ch = *p++;
+          while (*p && *p != ch)
+            p += (*p == '\\' ? 2 : 1);
+          if (*p == ch)
+            ++p;
+        } else if (char_is_pp_identifier_first(*p)) {
+          size_t len;
+          char *s = scan_pp_identifier(p);
+          len = s-p;
+          found_parameter = ((int) (size_t) stringhash_get(&parameters, p, len))-1;
+          if (found_parameter >= 0) {
+            end = p;
+            p = s;
+            break;
+          }
+          p = s;
+        } else if (*p == '#') {
+          end = p;
+          if (p[1] == '#') {
+            require_identifier = 1;
+            token_paste_after_text = 1;
+            p = preprocessor_skip_whitespace_simple(p+2);
+            if (*p == '#') {
+              stringify = 1;
+              ++p;
+            }
+            break;
+          }
+          require_identifier = 1;
+          stringify = 1;
+          ++p;
+          break;
+        } else if (char_is_end_of_line(*p)) {
+          end = p;
+          break;
+        } else
+          ++p;
+      }
+      if (char_is_end_of_line(*p) && found_parameter < 0) {
+        e.argument_index = -1;
+        while (start < p && char_is_whitespace(*start))
+          ++start;
+        e.text = arena_alloc_padded_string(&c->macro_arena, start, p-start);
+        e.text_length = (uint32) (p-start+2);
+        // always trim final space
+        e.text[--e.text_length] = 0;
+        if (token_paste_before_text) {
+          ++e.text; // delete space at front so token-pasting occurs
+          --e.text_length;
+        }
+        arrput(m->expansion, e);
+        break;
+      }
+      if (found_parameter < 0) {
+        size_t len;
+        char *s;
+        assert(require_identifier);
+        p = preprocessor_skip_whitespace_simple(p);
+        if (!char_is_pp_identifier_first(*p))
+          goto expected_parameter_name;
+        s = scan_pp_identifier(p);
+        len = s-p;
+        found_parameter = ((int) (size_t) stringhash_get(&parameters, p, len))-1;
+        p = s;
+        if (found_parameter < 0)
+          goto expected_parameter_name;
+      }
+
+      while (char_is_whitespace(*start) && start < end)
+        ++start;
+      while (end-1 > start && char_is_whitespace(end[-1]))
+        --end;
+
+      // if the string is empty, then a token paste is between two parameters
+      if (start == end && token_paste_before_text)
+        token_paste_after_text = 1;
+
+      e.argument_index = (int16) found_parameter;
+      e.stringify_argument = (uint8) stringify;
+      e.concat_argument = (uint8) token_paste_after_text;
+      e.eliminate_comma_before_empty_arg = 0;
+      e.text = arena_alloc_padded_string(&c->macro_arena, start, end-start);
+      e.text_length = (uint32) (end-start+2);
+
+      // to token paste with following parameter, remove space at end of expansion text
+      if (token_paste_after_text) {
+        e.text[--e.text_length] = 0;
+      }
+
+      // to token paste with previous parameter, remove space at start of expansion text
+      if (token_paste_before_text) {
+        ++e.text;
+        e.text_length -= 1;
+      }
+
+      arrput(m->expansion, e);
+
+      // check if the argument is followed by ##
+      p = preprocessor_skip_whitespace_simple(p);
+      if (p[0] == '#' && p[1] == '#') {
+        // if so, we want to token paste between this argument and the next text, i.e. before the next text
+        token_paste_before_text = 1;
+        arrlast(m->expansion).concat_argument = 1;
+        p += 2;
+      } else
+        token_paste_before_text = 0;
+    }
+
+    stringhash_destroy(&parameters);
+  }
+
+  return m;
+
+  error_exit:
+  {
+    return NULL;
+  }
+
+  invalid_char_error:
+  {
+    char cbuf[12];
+    char *error_messages[3] =
+    {
+      "Invalid character %s after '#define'",
+      "Invalid character %s in macro definition after '%s'",
+      "Invalid character %s in parameter list for macro definition '%s'",
+    };
+
+    char bad = *p;
+
+    do_error_code(ps, PP_RESULT_invalid_char, error_messages[state], error_char(bad,cbuf), m->symbol_name);
+
+    return NULL;
+  }
+
+  expected_parameter_name:
+  {
+    do_error(ps, "Expected a parameter name to follow # character in definition of macro '%s'.", m->symbol_name);
+
+    return NULL;
+  }
+}
+
+#if 0
+// originally, we tried to use a fast argument scanner that didn't make a copy.
+// however, we need to append a NUL to the end of the argument so that we can
+// feed it to the macro scanner, and arguments can come from the input file.
+// (we can't overwrite that in the middle with the NUL since it will be memory-mapped
+// in the common case)
+//
+// so instead we always copy the argument. might be worth coming back to this and
+// seeing if it's better to never memory map and instead apply this logic!
+static char *parse_argument(char *text, int *line_number, int *needs_copying)
+{
+  int num_parens=0;
+  char *p = text;
+  for(;;) {
+    while (!char_is_argument_parse_action((uint8) *p))
+      ++p;
+    switch (*p) {
+      case '\n': case '\r':
+        p += newline_char_count(p[0],p[1]);
+        ++*line_number;
+        break;
+      case ',':
+        if (num_parens == 0)
+          return p;
+        ++p;
+        break;
+      case ')':
+        if (num_parens == 0)
+          return p;
+        --num_parens;
+        ++p;
+        break;
+      case '(':
+        ++num_parens;
+        ++p;
+        break;
+      case '\'':
+      case '\"': {
+        char ch = *p++;
+        for(;;) {
+          while (!char_is_special_string_contents(*p))
+            ++p;
+          if (*p == ch || *p == 0)
+            break;
+          if (*p == '\\')
+            ++p; // skip character after backslash in case it's a quote mark
+          ++p;
+        }
+        if (*p == 0) {
+          return p;
+        }
+        assert(*p == ch);
+        ++p;
+        break;
+      }
+      case '\\':
+        *needs_copying = 1;
+        return p;
+      case '/':
+        if (p[1] == '/' || p[1] == '*') {
+          *needs_copying = 1;
+          return p;
+        } else
+          p += 1;
+        break;
+      default:
+        assert(0);
+      case 0:
+        return p;
+    }
+  }
+}
+#endif
+
+static char *copy_argument(char *text, int *line_number, char **p_out)
+{
+  int num_parens=0;
+  char *p = text;
+  char *out = *p_out; // need to be able to append to existing buffer for VA_ARGS
+  arrsetcap(out, 200);
+
+  for(;;) {
+    // scan characters that don't need processing
+    // we don't have a length bound, so we don't copy as we go, we make a second pass
+    char *q = p;
+    ptrdiff_t addlen;
+    size_t oldlen;
+    while (!char_is_argument_parse_action((uint8) *p)) {
+      ++p;
+    }
+    oldlen = arrlen(out);
+    addlen = p - q;
+    arrsetlen(out, oldlen+addlen);
+    memcpy(out+oldlen, q, addlen);
+
+    switch (*p) {
+      case '\n': case '\r':
+        if (p[-1] == '\\')
+          arrsetlen(out, arrlen(out)-1); // undo output of '\'
+        else
+          arrput(out, ' ');
+        p += newline_char_count(p[0],p[1]);
+        ++*line_number;
+        break;
+      case ',':
+        if (num_parens == 0) {
+          *p_out = out;
+          return p;
+        }
+        arrput(out,*p);
+        ++p;
+        break;
+      case ')':
+        if (num_parens == 0) {
+          *p_out = out;
+          return p;
+        }
+        --num_parens;
+        arrput(out,*p);
+        ++p;
+        break;
+      case '(':
+        ++num_parens;
+        arrput(out,*p);
+        ++p;
+        break;
+      case '\'':
+      case '\"': {
+        char ch = *p++;
+        arrput(out,ch);
+        for(;;) {
+          while (!char_is_special_string_contents(*p)) {
+            arrput(out,*p);
+            ++p;
+          }
+          if (*p == ch || *p == 0)
+            break;
+          if (*p == '\\') {
+            arrput(out,*p);
+            ++p;
+          }
+          arrput(out,*p);
+          ++p;
+        }
+        if (*p == 0) {
+          return p;
+        }
+        assert(*p == ch);
+        arrput(out,*p);
+        ++p;
+        break;
+      }
+      case '/':
+        if (p[1] == '/') {
+          while (!char_is_end_of_line(*p))
+            ++p;
+          if (*p == 0)
+            return p;
+          p += newline_char_count(p[0], p[1]);
+          arrput(out, ' ');
+          ++*line_number;
+        } else if (p[1] == '*') {
+          p += 2;
+          while (p[0] && !(p[0] == '*' && p[1] == '/')) {
+            if (*p == '\n' || *p == '\r') {
+              ++*line_number;
+              p += newline_char_count(p[0], p[1]);
+            } else {
+              ++p;
+            }
+          }
+          if (*p == 0)
+            return p;
+          p += 2;
+          arrput(out, ' ');
+        } else {
+          arrput(out,*p);
+          p += 1;
+        }
+        break;
+      default:
+        assert(0);
+      case 0: {
+        *p_out = out;
+        // allow caller to handle the error
+        return p;
+      }
+    }
+  }
+}
+
+// if the identifier is a function-type macro but it is followed by EOF, then the parent needs to check
+// if it's followed by '('
+static void maybe_expand_macro(parse_state *cs, macro_definition *pending, int in_macro_expansion, char *in_macro, macro_definition **final_function_macro)
+{
+  pp_context *c = cs->context;
+  macro_definition *md;
+  char *identifier = cs->fast_dest ? cs->fast_dest : (char *) cs->dest + arrlen(cs->dest);
+  size_t identifier_length;
+  char *p = cs->src + cs->src_offset;
+  char *q = identifier;
+  uint8 c0 = (uint8) p[0];
+  uint8 c1 = (uint8) p[1];
+
+  if (pending) {
+    md = pending;
+    identifier_length = 0; // NOTUSED
+  } else {
+    char *start = p;
+
+    #ifndef DISABLE_SHORT_MACRO_FILTER
+    if (!char_is_pp_identifier(c1)) {
+      // 1-character identifier
+      if (c->macro_name_filter[c0][0] == 0) {
+        *q++ = c0;
+        p += 1;
+        goto not_macro;
+      }
+    } else if (!char_is_pp_identifier(p[2])) {
+      // 2-character identifier
+      // could use bloom filter instead, but while this is 64KB, only portions of it will ever
+      if (c->macro_name_filter[c0][c1] == 0) {
+        *q++ = c0;
+        *q++ = c1;
+        p += 2;
+        goto not_macro;
+      }
+    }
+    #endif
+
+    // copy the rest of the identifier so we can NUL-terminate it
+    // we copy it to the output buffer since we've already allocated space for it,
+    // and that's where we'll need to copy it if it's not a macro
+    while (char_is_pp_identifier(*p))
+      *q++ = *p++;
+
+    identifier_length = p-start;
+    md = stringhash_get(&c->macro_map, identifier, identifier_length);
+    if (md == NULL || md->disabled) {
+      not_macro:
+        cs->src_offset = p - cs->src;
+        if (cs->fast_dest)
+          cs->fast_dest = q;
+        else
+          arrsetlen(cs->dest, (size_t) (q - cs->dest));  // set the output buffer size to account for the above-copied identifier
+        return;
+    }
+
+    cs->src_offset = p - cs->src;
+  }
+
+  if (md->num_parameters < 0) {
+    // no parens, no substitutions, so just process the expansion without copying
+
+    switch (md->num_parameters) {
+      case MACRO_NUM_PARAMETERS_no_parentheses: {
+        if (md->simple_expansion_length == 0)
+          // normally we avoid token pasting by pre-adding spaces to expansion. but if client
+          // passes in their own macro definitions, we let them pass truly empty strings,
+          // and prevent token pasting here; otherwise e.g.
+          //   foo.c:              +FOO+x;
+          //   command:            cpp -DFOO foo.c
+          //   desired output:     + +x;
+          //   without next line:  ++x;
+          if (cs->fast_dest)
+            *cs->fast_dest++ = ' ';
+          else
+            arrput(cs->dest, ' ');
+        else {
+          parse_state ncs = *cs;
+          ncs.parent = cs;
+          ncs.src = md->simple_expansion;
+          ncs.src_offset = 0;
+          ncs.src_length = md->simple_expansion_length;
+          md->disabled = 1;
+          preprocess_string(&ncs, IN_MACRO_yes, md->symbol_name, final_function_macro);
+          md->disabled = 0;
+          cs->dest = ncs.dest;
+          cs->fast_dest = ncs.fast_dest;
+        }
+        break;
+      }
+      case MACRO_NUM_PARAMETERS_defined: {
+        int parenthesized=0;
+        size_t len;
+
+        if (in_macro_expansion != IN_MACRO_if_condition)
+          goto not_macro;
+
+        // copy through the following token as well
+        while (char_is_whitespace(*p))
+          ++p;
+        if (*p == '(') {
+          ++p;
+          parenthesized = 1;
+          while (char_is_whitespace(*p))
+            ++p;
+        }
+        if (!char_is_pp_identifier_first(*p)) {
+          do_error(cs, "expected macro name after keyword 'defined'");
+        }
+        // rewind the output and copy the next identifier to the output temporarily
+        q = identifier;
+        do
+          *q++ = *p++;
+        while (char_is_pp_identifier(*p));
+        *q = 0;
+        len = q - identifier;
+
+        // rewind the output to prepare to write 0 or 1
+        q = identifier;
+        if (stringhash_get(&c->macro_map, identifier, len))
+          *q++ = '1';
+        else
+          *q++ = '0';
+
+        if (parenthesized) {
+          while (char_is_whitespace(*p))
+            ++p;
+          if (*p != ')') {
+            char cbuf[12];
+            do_error(cs, "unexpected character %s at end of 'defined('", error_char(*p,cbuf));
+            return;
+          }
+          ++p;
+          *q++ = ' '; // avoid token-pasting 
+        }
+        cs->src_offset = p - cs->src;
+        arrsetlen(cs->dest, (size_t) (q - cs->dest));  // set the output buffer size to account for the above-copied identifier
+        break;
+      }
+
+      case MACRO_NUM_PARAMETERS_file: {
+        size_t len = strlen(cs->filename);
+        // @TODO: backslash escape filename and cache it in stack, ready for string printing... also use when printing the line number?
+        if (cs->fast_dest) {
+          *cs->fast_dest++ = '"';
+          cs->fast_dest += sprintf(cs->fast_dest, "%s", cs->filename);
+          *cs->fast_dest++ = '"';
+        }  else {
+          size_t off = arraddnindex(cs->dest, len+2);
+          cs->dest[off] = '"';
+          memcpy(cs->dest+off+1, cs->filename, len);
+          cs->dest[off+1+len] = '"';
+        }
+        break;
+      }
+      case MACRO_NUM_PARAMETERS_line: {
+        if (cs->fast_dest)
+          cs->fast_dest += sprintf(cs->fast_dest, "%u", cs->src_line_number);
+        else {
+          int len;
+          q = identifier;
+          arrsetcap(cs->dest, arrlen(cs->dest)+12);
+          len = sprintf(cs->dest+arrlen(cs->dest), "%u", cs->src_line_number);
+          arrsetlen(cs->dest, arrlen(cs->dest)+len);
+        }
+        break;
+      }
+      case MACRO_NUM_PARAMETERS_counter: {
+        if (c->counter == 2147483647) {
+          // we stop after 31-bits just so we don't have to deal with printing more than 10 digits
+          do_error_code(cs, PP_RESULT_counter_overflowed, "__COUNTER__ was used too many times (%u)", c->counter);
+          return;
+        }
+        if (cs->fast_dest)
+          cs->fast_dest += sprintf(cs->fast_dest, "%u", c->counter);
+        else {
+          int len;
+          q = identifier;
+          arrsetcap(cs->dest, arrlen(cs->dest)+12);
+          len = sprintf(cs->dest+arrlen(cs->dest), "%u", c->counter);
+          arrsetlen(cs->dest, arrlen(cs->dest)+len);
+        }
+        ++c->counter;
+        break;
+      }
+    }
+  } else {
+    // macro has arguments, needs parenthesis
+    parse_state ncs = { 0 }, acs = { 0 };
+    char **arguments = NULL;
+    int i;
+    int temp_lines;
+    int parse_rest=0;
+    char *s = p;
+
+    // plan:
+    //   1. find opening parenthesis, if any
+    //   2. parse arguments into arguments[] array
+    //   3. parse __VA_ARGS__ and put at end of arguments[] array
+    //   4. build temporary buffer containing macro output
+    //   5. preprocess temporary buffer
+
+    //////////////////////////////////////////////////////////////////
+    //
+    //   1. find opening parenthesis, if any
+    //
+
+    temp_lines = 0;
+    s = preprocessor_skip_whitespace(s, &temp_lines);
+
+    // if macro name appears without following parenthesis, pass it through unaltered
+    if (*s != '(') {
+      if (*s == 0 && in_macro_expansion == IN_MACRO_yes && final_function_macro) {
+        // except if it's the last thing in a macro expansion
+        // skip everything we've seen
+        cs->src_offset = s - cs->src;
+        *final_function_macro = md;
+        return;
+      }
+      // @OPTIMIZE: don't reparse what we just scanned
+      goto not_macro;
+    }
+    cs->src_line_number += temp_lines;
+
+    // consume all the scanned characters
+    p = s+1;
+
+    #if 0
+    // WHAT?!?
+    // undo '(' from output
+    if (cs->fast_dest)
+      --cs->fast_dest;
+    else
+      arrsetlen(cs->dest, arrlen(cs->dest)-1);
+    #endif
+
+    //////////////////////////////////////////////////////////////////
+    //
+    //   2. parse arguments into arguments[] array
+    //
+
+    for (i=0; i < md->num_parameters; ++i) {
+      char *copy = 0;
+      int arg_newlines = 0;
+
+      p = preprocessor_skip_whitespace(p, &cs->src_line_number);
+
+      s = copy_argument(p, &arg_newlines, &copy);
+
+      // examine character after argument to make sure it's right character
+      // the only thing that should have stopped argument parsing is ',' ')' or EOF
+      if (*s == ',') {
+        if (i < md->num_parameters-1)
+          ; // ok, more arguments expected
+        else if (md->is_variadic)
+          ; // ok, parsed final argument but macro is variadic
+        else {
+          int err = do_error_code(cs, PP_RESULT_too_many_arguments_to_macro, "Too many arguments supplied to macro '%s'; only %d expected", md->symbol_name, md->num_parameters);
+          if (err)
+            return;
+          parse_rest = 1;
+        }
+      } else if (*s == ')') {
+        if (i < md->num_parameters-1) {
+          do_error(cs, "Too few arguments supplied to macro '%s'; got %d, expected %d", md->symbol_name, i, md->num_parameters);
+          return;
+        }
+      } else {
+        assert(*s == 0);
+        if (in_macro_expansion == IN_MACRO_if_condition || in_macro_expansion == IN_MACRO_directive)
+          do_error(cs, "Macro '%s' argument list unterminated in %s", md->symbol_name, in_macro);
+        else if (in_macro)
+          do_error(cs, "Macro '%s' argument list unterminated in expansion of outer macro '%s'", md->symbol_name, in_macro);
+        else
+          do_error(cs, "End-of-file in macro '%s' argument list", md->symbol_name);
+        return;
+      }
+
+      {
+        char *e = s;
+        // need to strip leading and trailing whitespace so token-pasting works
+
+        arrput(copy,0);
+        p = copy;
+        e = copy + arrlen(copy) - 1; // get address of NUL
+
+        if (e-1 > p) {
+          e = preprocessor_skip_whitespace_reverse(e);
+          assert(e > p);
+          if (e < p) e = p; // this should be impossible
+        }
+        *e = 0;
+
+        arrsetlen(copy, e-p);
+        arrput(arguments, copy);
+        cs->src_line_number += arg_newlines;
+
+        p = s+1; // advance to after terminating character
+      }
+    }
+
+    if (c->error) {
+      arrfree(arguments);
+      return;
+    }
+
+    //////////////////////////////////////////////////////////////////
+    //
+    //   3. parse __VA_ARGS__ and put at end of arguments[] array
+    //
+
+    if (md->is_variadic || parse_rest) {
+      // create n+1th argument
+      // check if it was already terminated with ')'
+      if (p[-1] == ')') {
+        // empty argument
+        arrput(arguments, NULL);
+      } else {
+        int arg_newlines=0;
+        char *e, *copy=0;
+        p = preprocessor_skip_whitespace(p, &cs->src_line_number);
+        // parse all the remaining arguments
+        s = p;
+        for(;;) {
+          arg_newlines = 0;
+          s = copy_argument(s, &arg_newlines, &copy);
+          if (*s == 0) {
+        if (in_macro_expansion == IN_MACRO_if_condition || in_macro_expansion == IN_MACRO_directive)
+              do_error(cs, "Macro '%s' argument list unterminated in %s", md->symbol_name, in_macro);
+            else if (in_macro)
+              do_error(cs, "Macro '%s' argument list unterminated in expansion of outer macro '%s'", md->symbol_name, in_macro);
+            else
+              do_error(cs, "End-of-file in macro '%s' argument list", md->symbol_name);
+            return;
+          }
+          if (*s == ')')
+            break;
+          arrput(copy, *s);
+          ++s;
+        }
+
+        // s points to closing parenthesis
+        assert(*s == ')');
+
+        arrput(copy, 0);
+        p = copy;
+        e = p + arrlen(p);
+
+        // strip trailing whitespace
+        while (char_is_whitespace(e[-1]))
+          --e;
+        if (e < p) e = p;
+        *e = 0;
+
+        arrput(arguments, p);
+        cs->src_line_number += arg_newlines;
+      }
+    }
+
+    // update the parse_state
+    cs->src_offset = (s+1) - cs->src;
+
+    //////////////////////////////////////////////////////////////////
+    //
+    //   4. build temporary buffer containing macro output:
+    //
+    //       for each block of raw text, just copy it
+    //       for each stringified or concatenated argument just copy it
+    //       for other arguments, perform an expansion now
+    //
+
+    MDBG();
+
+    // once we build the concatenated temp buffer, we will rescan that.
+    // so that temp buffer is stored in 'ncs'
+    ncs = *cs; // inherit file & line and destination
+    ncs.parent = cs;
+
+    // each argument will be parsed using acs, so set that up now
+    acs = *cs;
+    acs.dest = acs.fast_dest = 0;
+
+    // populate temporary buffer
+    s = 0;
+    for (i=0; i < arrlen(md->expansion)-1; ++i) {
+      int an;
+      size_t off;
+
+      off = arraddnindex(s, md->expansion[i].text_length);
+      // copy macro-expansion text
+      memcpy(s+off, md->expansion[i].text, md->expansion[i].text_length);
+
+      an = md->expansion[i].argument_index;
+
+      if (md->expansion[i].eliminate_comma_before_empty_arg) {
+        // this flag is normally only set if both of the following are true:
+        //    1. the previous expansion ends with a comma
+        //    2. this argument is __VA_ARGS__
+        // however, you could use it with non-__VA_ARGS__ and it will still work,
+        // if you wanted to make some new extension
+
+        if (arguments[an] == NULL) {
+          // find the generated comma and set it to whitespace
+          char *t = s+arrlen(s);
+          while (*--t != ',')
+            ;
+          *t = ' ';
+        }
+      }
+
+      // copy argument substitution
+
+      if (md->expansion[i].stringify_argument || md->expansion[i].concat_argument) {
+        if (md->expansion[i].stringify_argument) {
+          arrput(s, ' ');
+          arrput(s, '"');
+          // @TODO: stringify needs to insert backslashes before \ and "
+          off = arraddnindex(s, arrlen(arguments[an]));
+          memcpy(s+off, arguments[an], arrlen(arguments[an]));
+          arrput(s, '"');
+        } else {
+          off = arraddnindex(s, arrlen(arguments[an]));
+          memcpy(s+off, arguments[an], arrlen(arguments[an]));
+        }
+      } else {
+        macro_definition *ffm = 0;
+        acs.src = arguments[an];
+        acs.src_length = arrlen(arguments[an]);
+        acs.dest = s;
+        acs.fast_dest = 0;
+        acs.src_offset = 0;
+        preprocess_string(&acs, in_macro_expansion, identifier, &ffm);
+        s = acs.dest;
+        if (ffm) {
+          // concatenate it into the result string... if it's at the end, we'll catch it again in the processing below
+          size_t len = strlen(ffm->symbol_name);
+          off = arraddnindex(s, len);
+          memcpy(s+off, ffm->symbol_name, len);
+        }
+        if (cs->context->stop)
+          return;
+      }
+    }
+
+    {
+      size_t off = arraddnindex(s, md->expansion[i].text_length);
+      memcpy(s+off, md->expansion[i].text, md->expansion[i].text_length);
+    }
+
+    arrput(s, 0);
+
+    {
+      for (i=0; i < arrlen(arguments); ++i)
+        arrfree(arguments[i]);
+      arrfree(arguments);
+    }
+
+    ncs.src = s;
+    ncs.src_length = arrlen(s);
+    ncs.src_offset = 0;
+
+    //////////////////////////////////////////////////////////////////
+    //
+    //   5. preprocess temporary buffer
+    //
+
+    md->disabled = 1; // don't disable earlier, because if you CALL foo() WITH foo, it should still get expanded
+    preprocess_string(&ncs, IN_MACRO_yes, identifier, final_function_macro);
+    cs->dest = ncs.dest;
+    cs->fast_dest = ncs.fast_dest;
+    arrfree(ncs.src);
+
+    md->disabled = 0;
+  }
+}
+
+//////////////////////////////////////////////////////////////////////////////
+//
+//  DIRECTIVE PROCESSING
+//
+
+// it's already been cleaned of backslash-continuation and is NUL-terminated
+static char *macro_expand_directive(parse_state *cs, char *p, char *directive)
+{
+  parse_state ncs;
+  ncs = *cs;
+  ncs.src = p;
+  ncs.src_offset = 0;
+  ncs.parent = cs;
+  ncs.dest = 0;
+  ncs.fast_dest = 0;
+  preprocess_string(&ncs, IN_MACRO_directive, directive, NULL);
+  return ncs.dest;
+}
+
+static int evaluate_if(parse_state *cs, char *p, int *syntax_error)
+{
+  int result=0;
+  parse_state ncs;
+  ncs = *cs;
+  ncs.src = p;
+  ncs.src_length = strlen(p);
+  ncs.src_offset = 0;
+  ncs.parent = cs;
+  ncs.dest = 0;
+  ncs.fast_dest = 0;
+  
+  preprocess_string(&ncs, IN_MACRO_if_condition, "#if", NULL);
+
+  // recursive descent evaluator
+  if (!cs->context->stop && ncs.dest != NULL) {
+    arrput(ncs.dest, '\0');
+    result = evaluate_integer_constant_expression_as_condition(ncs.dest, syntax_error);
+  } else {
+    *syntax_error = 0;
+  }
+  return result;
+}
+
+static char *parse_pp_identifier_in_directive(parse_state *ps, char *text, char *identifier, size_t identifier_buffer_length, char *directive, size_t *idlen)
+{
+  size_t i;
+  char *p = preprocessor_skip_whitespace_simple(text);
+  *idlen = 0;
+  identifier[0] = 0;
+  if (!char_is_pp_identifier_first(*p)) {
+    char cbuf[12];
+    do_error_code(ps, PP_RESULT_invalid_char, "Invalid character %s after '#%s'", error_char(*p,cbuf), directive);
+    return p;
+  }
+  identifier[0] = *p++;
+  for (i=1; i < identifier_buffer_length-1; ++i) {
+    if (!char_is_pp_identifier(*p))
+      break;
+    identifier[i] = *p++;
+  }
+  *idlen = i;
+  identifier[i] = 0;
+  if (char_is_pp_identifier(*p))
+    do_error_code(ps, PP_RESULT_identifier_too_long, "'#%s' identifier too long", directive);
+  return p;
+}
+
+enum
+{
+  // perfect hash function: sum of ASCII values mod 32
+  HASH_elif   =  0,
+  HASH_include=  4,
+  HASH_endif  =  6,
+  HASH_line   =  8,
+  HASH_else   =  9,
+  HASH_error  = 10,
+  HASH_define = 11,
+  HASH_if     = 15,
+  HASH_ifndef = 12,
+  HASH_undef  = 18,
+  HASH_pragma = 24,
+  HASH_ifdef  = 30,
+
+  HASH_none   = 31,
+};
+
+// normally this comes from a safe buffer, but during disabled-by-#if scanning it's an uncopied directive
+static int parse_directive_after_hash(char **ptr_p, int *newlines)
+{
+  int sum=0;
+  char *p = *ptr_p;
+  char *start;
+  p = preprocessor_skip_whitespace_in_directive(p, newlines);
+  start = p;
+  while (char_is_pp_identifier_first(*p))
+    sum += (uint8) *p++;
+  sum &= (DIRECTIVE_HASH_SIZE-1);
+  // @OPTIMIZE: use a masked 8-byte comparison?
+  if (0 == memcmp(start, directive_hash[sum].name, directive_hash[sum].name_len) && !char_is_pp_identifier(directive_hash[sum].name_len)) {
+    *ptr_p = p;
+    return sum;
+  } else {
+    *ptr_p = start;
+    return HASH_none;
+  }
+}
+
+enum {
+  RETURN_BEHAVIOR_return,
+  RETURN_BEHAVIOR_break,
+  RETURN_BEHAVIOR_skip_to_end_of_line,
+};
+
+static char *preprocess_string_from_file(parse_state *ps, char *filename, char *text, size_t textlen, int conditional_compilation_nesting, int is_header);
+static char *resolve_file(char path[4096], parse_state *ps, ppbool check_source_dirs, ppbool check_normal_includes, ppbool check_system_includes, char *filename, ptrdiff_t filename_len);
+static char *load_file_with_nul_terminator(char *filename, size_t *ptr_len, int *mode);
+static void load_file_free(char *filedata, int mode);
+
+static void pragma_once(pp_context *c, char *filename, char *macroname)
+{
+  if (macroname)
+    macroname = stb_arena_alloc_string(&c->macro_arena, macroname);
+  else
+    macroname = (char *) (size_t) 1;
+  shput(c->once_map, filename, macroname);
+}
+
+static int process_include(parse_state *cs, char *start, conditional_state *cons)
+{
+  int ret = RETURN_BEHAVIOR_break;
+  pp_context *c = cs->context;
+  char *filename, *p;
+  p = preprocessor_skip_whitespace_simple(start);
+  // @TODO: macro expansion
+  if (!(*p == '<' || *p == '"')) {
+    char cbuf[12];
+    do_error(cs, "Expected filename after '#include', found character %s", error_char(*p, cbuf));
+    return RETURN_BEHAVIOR_return;
+  } else {
+    char *once;
+    char path[4096];
+    char *start_char = p+1, *old_p;
+    char end_char = (*p == '<') ? '>' : '"';
+    ++p;
+    while (*p != end_char) {
+      if (*p == end_char)
+        break;
+      if (char_is_end_of_line(*p)) {
+        char cbuf[2] = { 0 };
+        cbuf[0] = *p;
+        do_error(cs, "Missing closing '%s' in '#include", cbuf);
+        return RETURN_BEHAVIOR_return;
+      }
+      ++p;
+    }
+
+    if (*p == '"') {
+      filename = resolve_file(path, cs, 1, 1, 0, start_char, p-start_char);
+    } else {
+      filename = resolve_file(path, cs, 0, 0, 1, start_char, p-start_char);
+    }
+
+    if (filename == NULL) {
+      memcpy(path, start_char, p-start_char);
+      path[p-start_char] = 0;
+      do_error(cs, "File '%s' not found", path);
+      return RETURN_BEHAVIOR_return;
+    }
+
+    old_p = p;
+    p = preprocessor_skip_whitespace_simple(p+1);
+    if (*p != 0) {
+      char cbuf[12];
+      do_error(cs, "Unexpected character %s after filename in '#include'", error_char(*p,cbuf));
+      return RETURN_BEHAVIOR_return;
+    }
+
+    once = shget(c->once_map, filename);
+    if (once == (char *) (size_t) 1) {
+      // #pragma once, so skip
+      c->num_onced_files += 1;
+    } else {
+      macro_definition *md = NULL;
+      
+      if (once != NULL)
+        md = stringhash_get(&c->macro_map, once, strlen(once));
+
+      // if macro is defined, it's an include guard, which means if it's defined we can skip processing
+      if (md == NULL) {
+        int mode;
+        char *data, *result;
+        size_t len;
+        data = load_file_with_nul_terminator(filename, &len, &mode);
+
+        if (data == NULL) {
+          fprintf(stderr, "Preprocessor fatal internal error: was able to open '%s' once, but not the second time.\n", filename);
+          exit(1);
+        }
+
+        c->num_includes += 1;
+        result = preprocess_string_from_file(cs, filename, data, len, cons->conditional_compilation_nesting_level, 1);
+        load_file_free(data, mode);
+        if (cs->fast_dest)
+          cs->fast_dest = result;
+        else
+          cs->dest = result;
+        if (c->stop)
+          return RETURN_BEHAVIOR_return;
+      } else {
+        c->num_onced_files += 1;
+      }
+    }
+  }
+  return ret;
+}
+
+// process # directives
+static void process_directive(parse_state *cs, conditional_state *cons)
+{
+  int n;
+  pp_context *c = cs->context;
+  char *p, *endptr;
+  char *alloc=NULL; // all exit paths SHOULD do 'if (alloc) arrfree(alloc)', but we allow it to leak on error-handling returns
+  char buffer[1000];
+  int dummy; // uninitialized, should never get incremented by parse_directive_after_hash
+
+  assert(cs->src[cs->src_offset] == '#');
+  p = copy_line_without_comments(buffer, sizeof(buffer), cs->src + cs->src_offset + 1, &cs->src_line_number, &endptr);
+  if (p != buffer)
+    alloc = p;
+
+  if (endptr == NULL) {
+    do_error(cs, "Unclosed '/*' comment contained in directive");
+    return;
+  }
+
+  cs->src_offset = endptr - cs->src;
+  // the above leaves src_offset pointing to the newline at the end of the directive
+
+  n = parse_directive_after_hash(&p, &dummy);
+
+  // when we reach here, n is the directive keyword, or if it's
+  // a value that's not defined as a directive, then it could be
+  // anything
+
+  switch (n & 31) { // mask to try to avoid range-checked switch
+    // not a known directive:
+    case HASH_none:
+    // random hash values which mostly don't match but could on garbage;
+    // included here so every value is specified to encourage a full jump table
+    case  1: case  2: case  3: case  5: case  7: case 13:
+    case 14: case 16: case 17: case 19: case 20: case 21:
+    case 22: case 23: case 25: case 26: case 27: case 28: case 29:
+    {
+      // p points to first newline or non-whitespace character after '#'
+      switch (*p) {
+        case '\r': case '\n':
+          // empty directive; just skip it, let regular code copy the newline
+          break;
+
+        default:
+          // invalid directive, produce a diagnostic
+          if (!char_is_pp_identifier_first(*p)) {
+            char cbuf[12];
+            do_error_code(cs, PP_RESULT_invalid_char, "Invalid character %s after '#'", error_char(*p,cbuf));
+          } else {
+            // if no warning, skip parsing
+            if (pp_result_mode[PP_RESULT_unrecognized_directive] != PP_RESULT_MODE_no_warning) {
+              static char *directives[] = {
+                "define", "else", "elif", "endif", "error", "if", "ifdef",
+                "ifndef", "include", "line", "pragma", "undef",
+              };
+              char identifier[64], *match = 0;
+              int i;
+              for (i=0; i < sizeof(identifier)-1; ++i) {
+                if (!char_is_pp_identifier(p[i]))
+                  break;
+                identifier[i] = p[i];
+              }
+              identifier[i] = 0;
+              if (i < sizeof(identifier) && pp_result_mode[PP_RESULT_unrecognized_directive] <= did_you_mean_threshold) {
+                match = find_closest_matching_string(identifier, directives, sizeof(directives)/sizeof(directives[0]),
+                                                       i > 7 ? 2 + i-7 : 2); // if the string is long, allow more errors, e.g. they put junk on the end
+              }
+              if (match)
+                do_error_code(cs, PP_RESULT_unrecognized_directive, "Unrecognized preprocessor directive '#%s', did you mean '%s'?", identifier, match);
+              else
+                do_error_code(cs, PP_RESULT_unrecognized_directive, "Unrecognized preprocessor directive '#%s'", identifier);
+            }
+          }
+          if (alloc) arrfree(alloc);
+          return;
+      }
+      break;
+    }
+
+    case HASH_line:
+      // @TODO
+      if (alloc) arrfree(alloc);
+      return;
+
+    case HASH_pragma: {
+      char identifier[256];
+      size_t idlen;
+
+      p = preprocessor_skip_whitespace_simple(p);
+      (void) parse_pp_identifier_in_directive(cs, p, identifier, sizeof(identifier), "pragma", &idlen);
+      // check for preprocessor tokens we recognize
+      switch (identifier[0]) {
+        case 'o':
+          if (0==strcmp(identifier, "once")) {
+            pragma_once(cs->context, cs->filename, NULL);
+            goto handled;
+          }
+          break;
+        case 'p':
+          if (0==strcmp(identifier, "push") || 0==strcmp(identifier, "push_macro")) {
+            // @TODO implement pragma push_macro
+            goto handled;
+          }
+          if (0==strcmp(identifier, "pop") || 0==strcmp(identifier, "pop_macro")) {
+            // @TODO implement pragma pop_macro
+            goto handled;
+          }
+          if (0==strcmp(identifier, "pack")) {
+            // do macro expansion to handle #pragma pack(push,_CRT_PACKING) in vadefs.h
+            //do_expansion:
+            {
+              // @TODO: this is macro-substituting "pack" as well
+              char *temp = macro_expand_directive(cs, p, "#pragma pack");
+              // replace the unexpanded directive text with the expanded directive text
+              if (alloc)
+                arrfree(alloc);
+              p = alloc = temp;
+            }
+          }
+          break;
+      }
+
+      // copy pragma through to output... existing code will output the newline, so don't do it here
+      if (cs->fast_dest) {
+        size_t out_length = strlen(p);
+        char *z = cs->fast_dest;
+        memcpy(z, "#pragma ", 8);
+        z += 8;
+        memcpy(z, p, out_length);
+        z += out_length;
+        cs->fast_dest = z;
+      } else {
+        size_t out_length = strlen(p);
+        size_t s = arraddnindex(cs->dest, 8);
+        memcpy(cs->dest+s, "#pragma ", 8);
+        s = arraddnindex(cs->dest, out_length);
+        memcpy(cs->dest+s, p, out_length);
+      }
+
+     handled:
+      if (alloc) arrfree(alloc);
+      return;
+    }
+
+    case HASH_define: {
+      macro_definition *old;
+      macro_definition *m;
+
+      m = create_macro_definition(cs, p);
+      if (c->include_guard_detect)
+        c->last_macro_definition = m;
+
+      if (m != NULL) {
+        old = stringhash_get(&c->macro_map, m->symbol_name, m->symbol_name_length);
+        if (old != NULL) {
+          // @TODO: compare new definition to old and produce error/warning
+          stringhash_delete(&c->macro_map, m->symbol_name, m->symbol_name_length);
+        }
+        stringhash_put(&c->macro_map, m->symbol_name, m->symbol_name_length, m);
+        if (m->symbol_name[1] == 0 || m->symbol_name[2] == 0)
+          c->macro_name_filter[(uint8) m->symbol_name[0]][(uint8) m->symbol_name[1]] = 1;
+      }
+      if (alloc) arrfree(alloc);
+      return;
+    }
+
+    case HASH_undef: {
+      char identifier[1025];
+      size_t idlen;
+      macro_definition *md;
+
+      // process #undef
+
+      p = parse_pp_identifier_in_directive(cs, p, identifier, sizeof(identifier), "undef", &idlen);
+      if (identifier[0] == 0)
+        if (alloc) arrfree(alloc);
+
+      md = stringhash_get(&c->macro_map, identifier, idlen);
+
+      if (md != NULL) {
+        if (md->predefined) {
+          if (do_error_code(cs, PP_RESULT_undef_of_predefined_macro, "Cannot #undef predefined macro '%s'", identifier))
+            return;
+        } else {
+          stringhash_delete(&c->macro_map, identifier, idlen);
+          shput(c->undef_map, identifier, (void*) (size_t) 1);
+        }
+      } else {
+
+        // VS header workaround
+        if (0==strcmp(identifier, "long"))
+          return;
+
+        // process #undef error
+        if (pp_result_mode[PP_RESULT_undef_of_undefined_macro] != PP_RESULT_MODE_no_warning) {
+          if (shget(c->undef_map, identifier)) {
+            // @DIAGNOSTIC: record the file/line of the last #undef, so we can report it here?
+            // however, recording #file requires making sure it stays alive, but we don't want to do it multiple times?
+            // then again #undefs are pretty rare, so maybe we can just arena alloc it on undef
+            if (do_error_code(cs,  PP_RESULT_undef_of_undefined_macro, "Tried to #undef '%s', which is not a defined macro name. However, it was defined previously.", identifier))
+              return;
+          } else if (pp_result_mode[PP_RESULT_undef_of_undefined_macro] <= did_you_mean_threshold) {
+            char *best=0;
+            int dist=9999,i;
+            int min_dist = n > 5 ? 3 : n > 3 ? 2 : 1;
+            n = (int) idlen;
+            if (n > 1) {
+              find_closest_match_init(&dist,&best, min_dist);
+              if (n >= 4) {
+                // in case it's a warning, we don't want to be too slow, so accelerate if we can
+                for (i=0; i < arrlen(c->macro_map.pair); ++i) {
+                  // quick cull... if the first four characters mismatch, it's not a candidate, since we only allow max 3 errors
+                  char *s = c->macro_map.pair[i].key;
+                  if (s) {
+                    if (s[0] && s[1] && s[2] && s[3] && s[0] != identifier[0] && s[1] != identifier[1] && s[2] != identifier[2] && s[3] != identifier[3])
+                      ;
+                    else
+                      find_closest_match_step(identifier,s, &dist,&best);
+                  }
+                }
+              } else {
+                for (i=0; i < arrlen(c->macro_map.pair); ++i)
+                  if (c->macro_map.pair[i].key)
+                    find_closest_match_step(identifier, c->macro_map.pair[i].key, &dist,&best);
+              }
+            }
+            if (dist <= min_dist ? do_error_code(cs,  PP_RESULT_undef_of_undefined_macro, "Tried to #undef '%s', which is not a defined macro name. Did you mean '%s'?", identifier, best)
+                         : do_error_code(cs,  PP_RESULT_undef_of_undefined_macro, "Tried to #undef '%s', which is not a defined macro name.", identifier))
+              return;
+          } else {
+            if (do_error_code(cs,  PP_RESULT_undef_of_undefined_macro, "Tried to #undef '%s', which is not a defined macro name.", identifier))
+              return;
+          }
+        }
+      }
+      if (alloc) arrfree(alloc);
+        return;
+    }
+
+    case HASH_include: {
+      process_include(cs, p, cons);
+      if (alloc) arrfree(alloc);
+      return;
+    }
+
+
+//////////////////////////////////////////////////////////////////////////////
+//
+//  CONDITIONAL PROCESSING
+//
+
+    case HASH_ifdef: {
+      ifdef_info ifinfo;
+      macro_definition *m;
+      char identifier[1025];
+      size_t idlen;
+
+      ifinfo.line = cs->src_line_number;
+      ifinfo.type = HASH_ifdef;
+      p = parse_pp_identifier_in_directive(cs, p, identifier, sizeof(identifier), "ifdef", &idlen);
+      if (c->stop)
+        return;
+
+      arrpush(c->ifdef_stack, ifinfo);
+      ++cons->conditional_compilation_nesting_level;
+
+      m = stringhash_get(&c->macro_map, identifier, idlen);
+      if (m == NULL || identifier[0] == 0)
+        cons->conditionally_disable = CONDITIONAL_skip_current_block;
+      else
+        cons->conditionally_disable = CONDITIONAL_none;
+
+      p = preprocessor_skip_whitespace_simple(p);
+      if (*p != 0) {
+        char cbuf[12];
+        do_error_code(cs, PP_RESULT_unexpected_character_after_end_of_directive, "Expected end of line in #ifdef, but found character %s", error_char(*p, cbuf));
+      }
+      break;
+    }
+
+    case HASH_ifndef: {
+      ifdef_info ifinfo;
+      macro_definition *m;
+      char identifier[1025];
+      size_t idlen;
+
+      ifinfo.line = cs->src_line_number;
+      ifinfo.type = HASH_ifndef;
+      p = parse_pp_identifier_in_directive(cs, p, identifier, sizeof(identifier), "ifndef", &idlen);
+      if (c->stop)
+        return;
+
+      if (c->include_guard_detect)
+        c->last_ifndef = stb_arena_alloc_string(&c->macro_arena, identifier);
+
+      arrpush(c->ifdef_stack, ifinfo);
+      ++cons->conditional_compilation_nesting_level;
+
+      m = stringhash_get(&c->macro_map, identifier, idlen);
+      if (m != NULL)
+        cons->conditionally_disable = CONDITIONAL_skip_current_block;
+      else
+        cons->conditionally_disable = CONDITIONAL_none;
+
+      p = preprocessor_skip_whitespace_simple(p);
+      if (*p != 0) {
+        char cbuf[12];
+        do_error_code(cs, PP_RESULT_unexpected_character_after_end_of_directive, "Expected end of line in #ifdef, but found character %s.", error_char(*p, cbuf));
+      }
+      break;
+    }
+
+    case HASH_if: {
+      ifdef_info ifinfo;
+      int result, syntax_error;
+
+      ifinfo.line = cs->src_line_number;
+      ifinfo.type = HASH_if;
+
+      p = preprocessor_skip_whitespace_simple(p);
+      result = evaluate_if(cs, p,  &syntax_error);
+
+      if (syntax_error) {
+        if (do_error(cs, "Syntax error in conditional expression of '#if'"))
+          return;
+        else
+          result = 0;
+      }
+
+      arrpush(c->ifdef_stack, ifinfo);
+      ++cons->conditional_compilation_nesting_level;
+
+      if (result)
+        cons->conditionally_disable = CONDITIONAL_none;
+      else
+        cons->conditionally_disable = CONDITIONAL_skip_current_block;
+      break;
+    }
+
+    case HASH_elif:
+    handle_elif_after_disabled_block:
+    {
+      ifdef_info ifinfo;
+      int result, syntax_error;
+
+      ifinfo.line = cs->src_line_number;
+      ifinfo.type = HASH_elif;
+
+      if (cons->conditional_compilation_nesting_level == 0) {
+        do_error(cs, "Found '#elif' without a prior '#if', '#ifdef', or '#ifndef' in the same file.");
+        return;
+      }
+
+      if (arrlast(c->ifdef_stack).type == HASH_else) {
+        do_error(cs, "Found '#elif' after previous '#else'");
+        error_supplement(cs, PP_RESULT_ERROR, arrlast(c->ifdef_stack).line, "Location of previous '#else'");
+        return;
+      }
+
+      if (cons->conditionally_disable == CONDITIONAL_none) {
+        // previous block was enabled, so we're done
+        cons->conditionally_disable = CONDITIONAL_skip_all_blocks;
+        break;
+      }
+
+      p = preprocessor_skip_whitespace_simple(p);
+      result = evaluate_if(cs, p, &syntax_error);
+      if (syntax_error) {
+        do_error(cs, "Syntax error in conditional expression of '#elif'");
+        return;
+      }
+
+      arrlast(c->ifdef_stack) = ifinfo;
+
+      if (result)
+        cons->conditionally_disable = CONDITIONAL_none;
+      else
+        cons->conditionally_disable = CONDITIONAL_skip_current_block;
+      break;
+    }
+
+    case HASH_else:
+    handle_else_after_disabled_block:
+    {
+      ifdef_info ifinfo;
+
+      p = preprocessor_skip_whitespace_simple(p);
+      ifinfo.type = HASH_else;
+      ifinfo.line = cs->src_line_number;
+
+      if (*p != 0) {
+        char cbuf[12];
+        if (do_error_code(cs, PP_RESULT_unexpected_character_after_end_of_directive, "Expected end of line in #else, but found character %s.", error_char(*p, cbuf)))
+          return;
+      }
+
+      if (cons->conditional_compilation_nesting_level == 0) {
+        do_error(cs, "Found '#else' without a prior '#if', '#ifdef', or '#ifndef' in the same file.");
+        return;
+      }
+
+      if (arrlast(c->ifdef_stack).type == HASH_else) {
+        // can't have two else clauses
+        do_error(cs, "Found '#else' after previous '#else'");
+        error_supplement(cs, PP_RESULT_ERROR, arrlast(c->ifdef_stack).line, "Location of previous '#else'");
+        return;
+      }
+
+      if (cons->conditionally_disable == CONDITIONAL_none) {
+        // previous block was enabled, so we're done
+        cons->conditionally_disable = CONDITIONAL_skip_all_blocks;
+      } else {
+        // previous block was disabled, so enable
+        cons->conditionally_disable = CONDITIONAL_none;
+      }
+
+      // replace old entry with new entry
+      arrlast(c->ifdef_stack) = ifinfo;
+      break;
+    }
+
+    handle_endif_after_disabled_block:
+    case HASH_endif: {
+
+      p = preprocessor_skip_whitespace_simple(p);
+      if (*p != 0) {
+        char cbuf[12];
+        if (do_error_code(cs, PP_RESULT_unexpected_character_after_end_of_directive, "Expected end of line in #endif, but found character %s.", error_char(*p, cbuf)))
+          return;
+      }
+
+      if (cons->conditional_compilation_nesting_level == 0) {
+        do_error(cs, "Found '#endif' without a prior '#if', '#ifdef', or '#ifndef' in the same file.");
+        return;
+      }
+      (void) arrpop(c->ifdef_stack);
+      --cons->conditional_compilation_nesting_level;
+      cons->conditionally_disable = CONDITIONAL_none;
+
+      if (cs->include_guard_candidate_macro_name != 0 && cons->conditional_compilation_nesting_level == cs->include_guard_endif_level) {
+        // found the #endif for a leading #ifndef/#define pair, so now check if there's
+        // any code after the #endif
+        int state;
+
+        state = scan_whitespace_and_comments(cs);
+
+        if (state == PP_WHITESCAN_STATE_stop_eof) {
+          // only whitespace and comments after #endif from include guard, so it's a real include guard!
+          //printf("Detected include guard: %s (%s)\n", cs->include_guard_candidate_macro_name, cs->filename);
+          pragma_once(cs->context, cs->filename, cs->include_guard_candidate_macro_name);
+        } else {
+          //printf("Invalid include guard: %s (%s)\n", cs->include_guard_candidate_macro_name, cs->filename);
+          rewind_to_start_of_line(cs);
+        }
+        cs->include_guard_candidate_macro_name = 0;
+        return;
+      }
+      break;
+    }
+
+    case HASH_error:
+      // @TODO
+      break;
+  }
+
+  // each case above individually scans to end of line so they can
+  // report the error unambiguously, so we should definitely be at
+  // the end of the line at this point
+
+  if (alloc)
+    arrfree(alloc);
+
+  if (cons->conditionally_disable == CONDITIONAL_none)
+    return;
+  else {
+    int disable_nesting_level = 0;
+    int newlines=0;
+
+    alloc = 0;
+
+    p = cs->src + cs->src_offset;
+
+    for(;;) {
+      char *q;
+      char *hash_location;
+      p = scan_to_directive(p, &newlines); // fast scanner that skips non-directive content
+
+      if (*p == 0) {
+        cs->src_line_number += newlines;
+        cs->num_disabled_lines += newlines;
+        {
+          int i;
+          for (i=(int) arrlen(c->ifdef_stack)-1; i >= cs->conditional_nesting_depth_at_start; --i) {
+            do_error(cs         ,                                           "No #endif found for previous #%s", directive_hash[c->ifdef_stack[i].type].name);
+            error_supplement(cs,  PP_RESULT_ERROR, c->ifdef_stack[i].line,  "Location of #%s"                 , directive_hash[c->ifdef_stack[i].type].name);
+          }
+        }
+        return;
+      }
+
+      assert(*p == '#');
+      hash_location = p;
+      q = p+1;
+      cs->src_line_number += newlines;
+      c->num_disabled_lines += newlines;
+      n = parse_directive_after_hash(&q, &cs->src_line_number);
+      p = q;
+      newlines = 0;
+
+      switch (n) {
+        ifdef_info ifinfo;
+
+        case HASH_if:
+        case HASH_ifdef:
+        case HASH_ifndef:
+          ifinfo.type = n;
+          ifinfo.line = cs->src_line_number;
+          ++disable_nesting_level;
+          arrpush(c->ifdef_stack, ifinfo);
+          break;
+
+        case HASH_else:
+          if (disable_nesting_level == 0) {
+            if (cons->conditionally_disable == CONDITIONAL_skip_current_block) {
+              p = copy_line_without_comments(buffer, sizeof(buffer), p, &cs->src_line_number, &endptr);
+              if (p != buffer)
+                alloc = p;
+              if (endptr == NULL) {
+                do_error(cs, "Unclosed '/*' comment contained in #else");
+                return;
+              }
+              cs->src_offset = endptr - cs->src;
+              goto handle_else_after_disabled_block;
+            }
+            if (arrlast(c->ifdef_stack).type == HASH_else) {
+              // we were in an else block, and now there's another else block, so produce an error
+              do_error(cs, "Multiple #else blocks before #endif.");
+              return;
+            }
+          }
+          arrlast(c->ifdef_stack).type = n;
+          arrlast(c->ifdef_stack).line = cs->src_line_number;
+          break;
+
+        case HASH_endif:
+          if (disable_nesting_level == 0) {
+            p = copy_line_without_comments(buffer, sizeof(buffer), p, &cs->src_line_number, &endptr);
+            if (p != buffer)
+              alloc = p;
+            if (endptr == NULL) {
+              do_error(cs, "Unclosed '/*' comment contained in #endif");
+              return;
+            }
+            cs->src_offset = endptr - cs->src;
+            goto handle_endif_after_disabled_block;
+          }
+          --disable_nesting_level;
+          arrpop(c->ifdef_stack);
+          break;
+
+        case HASH_elif:
+          if (disable_nesting_level == 0) {
+            if (cons->conditionally_disable == CONDITIONAL_skip_current_block) {
+              p = copy_line_without_comments(buffer, sizeof(buffer), p, &cs->src_line_number, &endptr);
+              if (p != buffer)
+                alloc = p;
+              if (endptr == NULL) {
+                do_error(cs, "Unclosed '/*' comment contained in #elif");
+                return;
+              }
+              cs->src_offset = endptr - cs->src;
+              goto handle_elif_after_disabled_block;
+            }
+          }
+          arrlast(c->ifdef_stack).type = HASH_elif;
+          arrlast(c->ifdef_stack).line = cs->src_line_number;
+          break;
+
+        default:
+          break;
+      }
+
+      for (;;) {
+        while (!char_is_end_of_line(*p))
+          ++p;
+        if (*p == 0) // if end of file, stop
+          break;
+        if (p[-1] != '\\') // if not preceded by backslash, stop
+          break;
+        ++newlines;
+        p += newline_char_count(p[0],p[1]);
+      }
+    }
+    /* NOTREACHED */  // end of infinite loop
+  }
+}
+
+int preprocessor_automatic_include_guard_detection=1;
+static char *preprocess_string_from_file(parse_state *ps, char *filename, char *text, size_t textlen, int conditional_compilation_nesting, int is_header)
+{
+  parse_state cs = *ps;
+  cs.parent = ps;
+  cs.filename = filename;
+  cs.last_output_filename = ""; // force #line directive to include filename
+  cs.src = text;
+  cs.src_line_number = 1;
+  cs.dest_line_number = 1;
+  cs.in_leading_whitespace = 1;
+  cs.src_length = textlen;
+  cs.src_offset =  0;
+  cs.conditional_nesting_depth_at_start = conditional_compilation_nesting;
+  // dest & fast_dest are inherited
+  preprocess_string(&cs, is_header && preprocessor_automatic_include_guard_detection ? IN_MACRO_include_guard_scan : IN_MACRO_no, 0, 0);
+  cs.context->num_lines += cs.src_line_number;
+  ps->last_output_filename = "";
+  if (cs.fast_dest)
+    return cs.fast_dest;
+  else
+    return cs.dest;
+}
+
+#define DEFAULT     -1
+#define ANY_LETTER  -2
+#define ANY_DIGIT   PP_CHAR_CLASS_digit
+
+#define STATE_ANY_WHITESPACE        -20
+#define STATE_ANY_MULTILINE_COMMENT -21
+#define STATE_DEFAULT               -22
+
+static void transition(int a, int b, int c)
+{
+  if (a == STATE_ANY_WHITESPACE) {
+    transition(PP_STATE_ready, b, c);
+    transition(PP_STATE_in_leading_whitespace, b, c);
+    transition(PP_STATE_saw_cr, b, c);
+    transition(PP_STATE_saw_lf, b, c);
+  } else if (a == STATE_ANY_MULTILINE_COMMENT) {
+    transition(PP_STATE_comment_multiline, b, c);
+    transition(PP_STATE_comment_multiline_seen_star, b, c);
+    transition(PP_STATE_comment_multiline_seen_cr, b, c);
+    transition(PP_STATE_comment_multiline_seen_lf, b, c);
+  } else if (a == STATE_DEFAULT) {
+    int i;
+    for (i=0; i < PP_STATE_active_count; ++i)
+      transition(i, b, c);
+  } else {
+    if (b == DEFAULT)
+      memset(pp_transition_table[a], c, PP_CHAR_CLASS_count);
+    else if (b == ANY_LETTER) {
+      transition(a        , PP_CHAR_CLASS_idtype       , c);
+      transition(a        , PP_CHAR_CLASS_idtype_e_E   , c);
+    }else
+      pp_transition_table[a][b] = (uint8) c;
+  }
+}
+
+static void initscan_transition(int a, int b, int c)
+{
+  if (a == STATE_ANY_MULTILINE_COMMENT) {
+    initscan_transition(PP_WHITESCAN_STATE_comment_multiline, b, c);
+    initscan_transition(PP_WHITESCAN_STATE_comment_multiline_seen_star, b, c);
+    initscan_transition(PP_WHITESCAN_STATE_comment_multiline_seen_cr, b, c);
+    initscan_transition(PP_WHITESCAN_STATE_comment_multiline_seen_lf, b, c);
+  } else if (a == STATE_DEFAULT) {
+    int i;
+    for (i=0; i < PP_WHITESCAN_STATE_active_count; ++i)
+      initscan_transition(i, b, c);
+  } else {
+    if (b == DEFAULT)
+      memset(pp_transition_initscan[a], c, PP_CHAR_CLASS_count);
+    else if (b == ANY_LETTER) {
+      initscan_transition(a        , PP_CHAR_CLASS_idtype       , c);
+      initscan_transition(a        , PP_CHAR_CLASS_idtype_e_E   , c);
+    } else
+      pp_transition_initscan[a][b] = (uint8) c;
+  }
+}
+
+void init_preprocessor_scanner(void)
+{
+  int i;
+
+  // build character class tables
+
+  assert(PP_CHAR_CLASS_count <= sizeof(pp_transition_table[0]));
+
+  memset(pp_char_class, PP_CHAR_CLASS_default, 256);
+
+  pp_char_class[ 0 ] = PP_CHAR_CLASS_eof;
+
+  for (i=0; i < 127; ++i) {
+    if (isalpha(i)) pp_char_class[i] = PP_CHAR_CLASS_idtype;
+    if (isdigit(i)) pp_char_class[i] = PP_CHAR_CLASS_digit;
+    if (isspace(i)) pp_char_class[i] = PP_CHAR_CLASS_whitespace;
+  }
+  assert(pp_char_class['\t'] == PP_CHAR_CLASS_whitespace);
+
+  pp_char_class['e'] = pp_char_class['E'] = PP_CHAR_CLASS_idtype_e_E,
+
+  pp_char_class['\r'] = PP_CHAR_CLASS_cr;
+  pp_char_class['\n'] = PP_CHAR_CLASS_lf;
+
+  pp_char_class['$'] = PP_CHAR_CLASS_idtype; // common C extension: allow $ in identifiers
+  pp_char_class['_'] = PP_CHAR_CLASS_idtype;
+   
+  pp_char_class['+'] = PP_CHAR_CLASS_plusminus;
+  pp_char_class['-'] = PP_CHAR_CLASS_plusminus;
+
+  pp_char_class['/'] = PP_CHAR_CLASS_slash;
+  pp_char_class['*'] = PP_CHAR_CLASS_star;
+  pp_char_class['.'] = PP_CHAR_CLASS_period;
+
+  pp_char_class['\"'] = PP_CHAR_CLASS_quotation;
+  pp_char_class['\''] = PP_CHAR_CLASS_apostrophe;
+  pp_char_class['\\'] = PP_CHAR_CLASS_backslash;
+
+  pp_char_class['@'] = PP_CHAR_CLASS_idtype; // RADC extension
+  pp_char_class['#'] = PP_CHAR_CLASS_hash;
+
+  // the scanner does the same as above, but doesn't stop on identifiers
+  memcpy(pp_scan_char_class, pp_char_class, 256);
+  for (i=0; i < 256; ++i)
+    if (pp_scan_char_class[i] == PP_CHAR_CLASS_idtype || pp_scan_char_class[i] == PP_CHAR_CLASS_idtype_e_E)
+      pp_scan_char_class[i] = PP_CHAR_CLASS_default;
+
+  // properties of transitions (stored as properties of end state)
+  memset(pp_state_is_end_of_line, 0, sizeof(pp_state_is_end_of_line));
+  pp_state_is_end_of_line[PP_STATE_saw_cr    ]                = 1;
+  pp_state_is_end_of_line[PP_STATE_saw_lf    ]                = 1;
+  pp_state_is_end_of_line[PP_STATE_comment_multiline_seen_cr] = 1;
+  pp_state_is_end_of_line[PP_STATE_comment_multiline_seen_lf] = 1;
+
+  // state machine transitions
+  // by default, all states point to eof so we catch missing transitions
+  for (i=0; i < PP_STATE_active_count; ++i)
+    transition(i, DEFAULT, PP_STATE_eof);
+
+  // every character class has a default behavior:
+  transition(STATE_DEFAULT           , PP_CHAR_CLASS_default      , PP_STATE_ready     );
+  transition(STATE_DEFAULT           , PP_CHAR_CLASS_whitespace   , PP_STATE_ready     );
+  transition(STATE_DEFAULT           , PP_CHAR_CLASS_cr           , PP_STATE_saw_cr    );
+  transition(STATE_DEFAULT           , PP_CHAR_CLASS_lf           , PP_STATE_saw_lf    );
+  transition(STATE_DEFAULT           , PP_CHAR_CLASS_period       , PP_STATE_in_number );
+  transition(STATE_DEFAULT           , PP_CHAR_CLASS_digit        , PP_STATE_in_number );
+  transition(STATE_DEFAULT           , PP_CHAR_CLASS_idtype       , PP_STATE_saw_identifier_start );
+  transition(STATE_DEFAULT           , PP_CHAR_CLASS_idtype_e_E   , PP_STATE_saw_identifier_start );
+  transition(STATE_DEFAULT           , PP_CHAR_CLASS_apostrophe   , PP_STATE_in_char_literal      );
+  transition(STATE_DEFAULT           , PP_CHAR_CLASS_quotation    , PP_STATE_in_string_literal    );
+  transition(STATE_DEFAULT           , PP_CHAR_CLASS_backslash    , PP_STATE_saw_backslash        );
+  transition(STATE_DEFAULT           , PP_CHAR_CLASS_slash        , PP_STATE_saw_slash );
+  transition(STATE_DEFAULT           , PP_CHAR_CLASS_star         , PP_STATE_ready     );
+  transition(STATE_DEFAULT           , PP_CHAR_CLASS_plusminus    , PP_STATE_ready     );
+  transition(STATE_DEFAULT           , PP_CHAR_CLASS_hash         , PP_STATE_saw_hash  );
+  transition(STATE_DEFAULT           , PP_CHAR_CLASS_eof          , PP_STATE_eof       );
+
+  // whitespace stays in leading whitespace
+  transition(PP_STATE_in_leading_whitespace, PP_CHAR_CLASS_whitespace , PP_STATE_in_leading_whitespace);
+
+  // anytime we see whitespace after a newline, we go to leading whitespace
+  transition(PP_STATE_saw_cr         , PP_CHAR_CLASS_whitespace  , PP_STATE_in_leading_whitespace);
+  transition(PP_STATE_saw_lf         , PP_CHAR_CLASS_whitespace  , PP_STATE_in_leading_whitespace);
+
+  // lf after cr
+  transition(PP_STATE_saw_cr         , PP_CHAR_CLASS_lf        , PP_STATE_in_leading_whitespace);
+#ifndef TREAT_LF_CR_AS_TWO_NEWLINES
+  transition(PP_STATE_saw_lf         , PP_CHAR_CLASS_cr        , PP_STATE_in_leading_whitespace);
+#endif
+
+  // leading whitespace tracking is so we can detect # after leading whitespace
+  transition(PP_STATE_in_leading_whitespace , PP_CHAR_CLASS_hash      , PP_STATE_saw_leading_hash);
+  transition(PP_STATE_saw_cr                , PP_CHAR_CLASS_hash      , PP_STATE_saw_leading_hash);
+  transition(PP_STATE_saw_lf                , PP_CHAR_CLASS_hash      , PP_STATE_saw_leading_hash);
+
+  // comments
+  transition(PP_STATE_saw_slash      , PP_CHAR_CLASS_slash     , PP_STATE_comment_singleline);
+  transition(PP_STATE_saw_slash      , PP_CHAR_CLASS_star      , PP_STATE_comment_multiline );
+
+  // single-line parse
+  transition(PP_STATE_comment_singleline, DEFAULT                 , PP_STATE_comment_singleline );
+  transition(PP_STATE_comment_singleline, PP_CHAR_CLASS_eof       , PP_STATE_eof );
+  transition(PP_STATE_comment_singleline, PP_CHAR_CLASS_cr        , PP_STATE_saw_cr);
+  transition(PP_STATE_comment_singleline, PP_CHAR_CLASS_lf        , PP_STATE_saw_lf);
+
+  // multi-line parse
+  transition(STATE_ANY_MULTILINE_COMMENT           , DEFAULT                , PP_STATE_comment_multiline);
+  transition(STATE_ANY_MULTILINE_COMMENT           , PP_CHAR_CLASS_cr       , PP_STATE_comment_multiline_seen_cr);
+  transition(STATE_ANY_MULTILINE_COMMENT           , PP_CHAR_CLASS_lf       , PP_STATE_comment_multiline_seen_lf);
+  transition(STATE_ANY_MULTILINE_COMMENT           , PP_CHAR_CLASS_star     , PP_STATE_comment_multiline_seen_star);
+  transition(STATE_ANY_MULTILINE_COMMENT           , PP_CHAR_CLASS_eof      , PP_STATE_error_unterminated_multiline_comment);
+  transition(PP_STATE_comment_multiline_seen_star  , PP_CHAR_CLASS_slash    , PP_STATE_ready);
+  transition(PP_STATE_comment_multiline_seen_cr    , PP_CHAR_CLASS_lf       , PP_STATE_comment_multiline); // don't double-count CR-LF
+#ifndef TREAT_LF_CR_AS_TWO_NEWLINES
+  transition(PP_STATE_comment_multiline_seen_lf   , PP_CHAR_CLASS_cr        , PP_STATE_comment_multiline); // don't double-count LF-CR
+#endif
+
+  // parse preprocessor number
+
+  transition(PP_STATE_in_number                   , PP_CHAR_CLASS_digit        , PP_STATE_in_number);
+  transition(PP_STATE_in_number                   , PP_CHAR_CLASS_period       , PP_STATE_in_number);
+  transition(PP_STATE_in_number                   , PP_CHAR_CLASS_idtype       , PP_STATE_in_number);
+  transition(PP_STATE_in_number                   , PP_CHAR_CLASS_idtype_e_E   , PP_STATE_in_number_saw_e);
+
+  transition(PP_STATE_in_number_saw_e             , PP_CHAR_CLASS_plusminus    , PP_STATE_in_number);
+  transition(PP_STATE_in_number_saw_e             , PP_CHAR_CLASS_digit        , PP_STATE_in_number);
+  transition(PP_STATE_in_number_saw_e             , PP_CHAR_CLASS_period       , PP_STATE_in_number);
+  transition(PP_STATE_in_number_saw_e             , PP_CHAR_CLASS_idtype       , PP_STATE_in_number);
+  transition(PP_STATE_in_number_saw_e             , PP_CHAR_CLASS_idtype_e_E   , PP_STATE_in_number_saw_e);
+
+  transition(PP_STATE_in_string_literal              , DEFAULT                    , PP_STATE_in_string_literal);
+  transition(PP_STATE_in_string_literal              , PP_CHAR_CLASS_quotation    , PP_STATE_ready);
+  transition(PP_STATE_in_string_literal              , PP_CHAR_CLASS_backslash    , PP_STATE_in_string_literal_saw_backslash);
+  transition(PP_STATE_in_string_literal_saw_backslash, DEFAULT                    , PP_STATE_in_string_literal);
+
+  transition(PP_STATE_in_char_literal               , DEFAULT                    , PP_STATE_in_char_literal);
+  transition(PP_STATE_in_char_literal               , PP_CHAR_CLASS_apostrophe   , PP_STATE_ready);
+  transition(PP_STATE_in_char_literal               , PP_CHAR_CLASS_backslash    , PP_STATE_in_char_literal_saw_backslash);
+  transition(PP_STATE_in_char_literal_saw_backslash , DEFAULT                    , PP_STATE_in_char_literal);
+
+  // init-scan state machine transitions
+  // by default, all states terminate scan
+  for (i=0; i < PP_STATE_active_count; ++i)
+    initscan_transition(i, DEFAULT, PP_WHITESCAN_STATE_stop_seen_onechar);
+
+  // every character class has a default behavior:
+  initscan_transition(STATE_DEFAULT           , PP_CHAR_CLASS_whitespace   , PP_WHITESCAN_STATE_whitespace);
+  initscan_transition(STATE_DEFAULT           , PP_CHAR_CLASS_cr           , PP_WHITESCAN_STATE_saw_cr    );
+  initscan_transition(STATE_DEFAULT           , PP_CHAR_CLASS_lf           , PP_WHITESCAN_STATE_saw_lf    );
+  initscan_transition(STATE_DEFAULT           , PP_CHAR_CLASS_slash        , PP_WHITESCAN_STATE_saw_slash );
+
+  // lf after cr
+  initscan_transition(PP_WHITESCAN_STATE_saw_cr         , PP_CHAR_CLASS_lf        , PP_WHITESCAN_STATE_whitespace);
+#ifndef TREAT_LF_CR_AS_TWO_NEWLINES
+  initscan_transition(PP_WHITESCAN_STATE_saw_lf         , PP_CHAR_CLASS_cr        , PP_WHITESCAN_STATE_whitespace);
+#endif
+
+  // comments
+  initscan_transition(PP_WHITESCAN_STATE_saw_slash      , DEFAULT                    , PP_WHITESCAN_STATE_stop_seen_twochars);
+  initscan_transition(PP_WHITESCAN_STATE_saw_slash      , PP_CHAR_CLASS_slash        , PP_WHITESCAN_STATE_comment_singleline);
+  initscan_transition(PP_WHITESCAN_STATE_saw_slash      , PP_CHAR_CLASS_star         , PP_WHITESCAN_STATE_comment_multiline );
+
+  // single-line parse
+  initscan_transition(PP_WHITESCAN_STATE_comment_singleline, DEFAULT                 , PP_WHITESCAN_STATE_comment_singleline );
+  initscan_transition(PP_WHITESCAN_STATE_comment_singleline, PP_CHAR_CLASS_cr        , PP_WHITESCAN_STATE_saw_cr);
+  initscan_transition(PP_WHITESCAN_STATE_comment_singleline, PP_CHAR_CLASS_lf        , PP_WHITESCAN_STATE_saw_lf);
+
+  // multi-line parse
+  initscan_transition(STATE_ANY_MULTILINE_COMMENT           , DEFAULT                , PP_WHITESCAN_STATE_comment_multiline);
+  initscan_transition(STATE_ANY_MULTILINE_COMMENT           , PP_CHAR_CLASS_cr       , PP_WHITESCAN_STATE_comment_multiline_seen_cr);
+  initscan_transition(STATE_ANY_MULTILINE_COMMENT           , PP_CHAR_CLASS_lf       , PP_WHITESCAN_STATE_comment_multiline_seen_lf);
+  initscan_transition(STATE_ANY_MULTILINE_COMMENT           , PP_CHAR_CLASS_star     , PP_WHITESCAN_STATE_comment_multiline_seen_star);
+  initscan_transition(PP_WHITESCAN_STATE_comment_multiline_seen_star, PP_CHAR_CLASS_slash    , PP_WHITESCAN_STATE_whitespace);
+  initscan_transition(PP_WHITESCAN_STATE_comment_multiline_seen_cr  , PP_CHAR_CLASS_lf       , PP_WHITESCAN_STATE_comment_multiline); // don't double-count CR-LF
+#ifndef TREAT_LF_CR_AS_TWO_NEWLINES
+  initscan_transition(PP_WHITESCAN_STATE_comment_multiline_seen_lf  , PP_CHAR_CLASS_cr        , PP_WHITESCAN_STATE_comment_multiline); // don't double-count LF-CR
+#endif
+
+  initscan_transition(STATE_DEFAULT                     , PP_CHAR_CLASS_eof      , PP_WHITESCAN_STATE_stop_eof );
+  initscan_transition(STATE_ANY_MULTILINE_COMMENT       , PP_CHAR_CLASS_eof      , PP_WHITESCAN_STATE_error_unterminated_multiline_comment);
+
+  // properties of transitions (stored as properties of end state)
+  memset(pp_initscan_state_is_end_of_line, 0, sizeof(pp_initscan_state_is_end_of_line));
+  pp_initscan_state_is_end_of_line[PP_WHITESCAN_STATE_saw_cr    ]                = 1;
+  pp_initscan_state_is_end_of_line[PP_WHITESCAN_STATE_saw_lf    ]                = 1;
+  pp_initscan_state_is_end_of_line[PP_WHITESCAN_STATE_comment_multiline_seen_cr] = 1;
+  pp_initscan_state_is_end_of_line[PP_WHITESCAN_STATE_comment_multiline_seen_lf] = 1;
+}
+
+macro_definition *pp_define(struct stb_arena *a, char *p)
+{
+  macro_definition *m;
+  pp_context c = { 0 };
+  parse_state ps = { 0 };
+  c.macro_arena = *a;
+  ps.filename = NULL;
+  ps.src_line_number = 0;
+  ps.parent = 0;
+  ps.context = &c;
+  ps.state_limit = PP_STATE_active_count;
+  m = create_macro_definition(&ps, p);
+  *a = c.macro_arena;
+  return m;
+}
+
+static void define_macro(pphash *map, macro_definition *m)
+{
+  stringhash_put(map, m->symbol_name, m->symbol_name_length, m);
+}
+
+char * preprocess_file(char *output_autobuffer, char *filename, char **include_paths, int num_include_paths, char **sys_include_paths, int num_sys_include_paths, macro_definition **predefined_macros, int num_predefined_macros, pp_diagnostic **pd, int *num_pd)
+{
+  char *output = 0;
+  pp_context c = { 0 };
+  int main_mode, i;
+  size_t length;
+  char *main_file;
+  
+  main_file = load_file_with_nul_terminator(filename, &length, &main_mode);
+  *num_pd = 0;
+
+  if (main_file == 0) {
+    parse_state ps = { 0 };
+    ps.filename = "<internal>";
+    ps.src_line_number = 0;
+    ps.context = &c;
+    do_error(&ps, "Couldn't open '%s'.", filename);
+    *pd = c.error;
+    *num_pd = (int) arrlen(c.error);
+    return NULL;
+  }
+
+  c.include_paths     = include_paths;
+  c.sys_include_paths = sys_include_paths;
+  c.num_include_paths = num_include_paths;
+  c.num_sys_include_paths = num_sys_include_paths;
+
+  // allocate the macro definition dictionary very large so (a) it's unlikely to grow,
+  // and (b) so most things hit in there first slot. but going too much bigger can
+  // backfire by slowing down the initialization
+  stringhash_create(&c.macro_map, 4096);
+
+  c.macro_name_filter = malloc(65536);
+  memset(c.macro_name_filter, 0, 65536);
+
+  for (i=0; i < arrlen(predefined_macros); ++i)
+    define_macro(&c.macro_map, predefined_macros[i]);
+
+  define_macro(&c.macro_map, &predefined_defined);
+  define_macro(&c.macro_map, &predefined_FILE);
+  define_macro(&c.macro_map, &predefined_LINE);
+  define_macro(&c.macro_map, &predefined_COUNTER);
+
+  for (i=0; i < arrlen(c.macro_map.pair); ++i)
+    update_macro_filter(&c, c.macro_map.pair[i].key);
+
+  {
+    // initial parse_state for preprocess_string_from_file to initialize from
+    parse_state ps = { 0 };
+    ps.context = &c;
+    ps.fast_dest = output_autobuffer;
+    ps.state_limit = PP_STATE_active_count;
+
+    output = preprocess_string_from_file(&ps, filename, main_file, length, 0, 0);
+    if (output_autobuffer) {
+      *output = 0;
+      output = output_autobuffer;
+    } else {
+      arrpush(output, 0);
+    }
+  }
+
+  load_file_free(main_file, main_mode);
+
+  shfree(c.once_map);
+  stringhash_destroy(&c.macro_map);
+  free(c.macro_name_filter);
+  stb_arena_free(&c.macro_arena);
+
+  if (c.error) {
+    *pd = c.error;
+    *num_pd = (int) arrlen(*pd);
+    if (c.stop) {
+      if (output_autobuffer == NULL)
+        arrfree(output);
+      output = NULL;
+    }
+  } else
+    *pd = 0;
+
+  //printf("%d lines (%d of them disabled by #if), not counting %d files skipped by #pragma once or include guards, %d includes\n", c.num_lines, c.num_disabled_lines, c.num_onced_files, c.num_includes);
+
+  return output;
+}
+
+void preprocessor_file_free(char *text, pp_diagnostic *pd)
+{
+  int i,j;
+  arrfree(text);
+  for (i=0; i < arrlen(pd); ++i) {
+    for (j=0; j < arrlen(pd[i].where); ++j)
+      free(pd[i].where[j].filename);
+    arrfree(pd[i].where);
+  }
+  arrfree(pd);
+}
+
+static void init_directive(char *s, int hash)
+{
+  int sum = 0;
+  int i;
+  for (i=0; s[i] != 0; ++i)
+    sum += (uint8) s[i];
+  sum &= (DIRECTIVE_HASH_SIZE-1);
+  if (sum != hash) {
+    fprintf(stderr, "Preprocessor fatal internal error: hash initialization mismatch for '%s' (%s %d).\n", s, __FILE__, __LINE__);
+    exit(1);
+  }
+  if (directive_hash[hash].name != 0) {
+    fprintf(stderr, "Preprocessor fatal internal error: hash collision between '%s' and '%s' (%s %d).\n", s, directive_hash[hash].name, __FILE__, __LINE__);
+    exit(1);
+  }
+
+  directive_hash[hash].name = s;
+  directive_hash[hash].name_len = strlen(s);
+}
+
+static void init_char_type(void)
+{
+  int i;
+  memset(pp_char_flags1, 0, 256);
+  memset(pp_char_flags2, 0, 256);
+
+  for (i=1; i < 127; ++i) {
+    if (isalpha(i) || i == '_' || i == '$')
+      pp_char_flags1[i] |= CHAR1_IS_PP_IDENTIFIER_FIRST;
+    if (isdigit(i))
+      pp_char_flags1[i] |= CHAR1_IS_DIGIT;
+  }
+
+  pp_char_flags1[ 0 ]  |= CHAR1_IS_NUL;
+  pp_char_flags1['\n'] |= CHAR1_IS_NEWLINE;
+  pp_char_flags1['\r'] |= CHAR1_IS_NEWLINE;
+  pp_char_flags1[' ' ] |= CHAR1_IS_NONNEWLINE_WHITESPACE;
+  pp_char_flags1['\t'] |= CHAR1_IS_NONNEWLINE_WHITESPACE;
+
+  pp_char_flags1['(' ] |= CHAR1_IS_BALANCE_SCAN_ACTION;
+  pp_char_flags1[')' ] |= CHAR1_IS_BALANCE_SCAN_ACTION;
+  pp_char_flags1[',' ] |= CHAR1_IS_BALANCE_SCAN_ACTION;
+  pp_char_flags1['/' ] |= CHAR1_IS_BALANCE_SCAN_ACTION;
+  pp_char_flags1['\\'] |= CHAR1_IS_BALANCE_SCAN_ACTION;
+
+  pp_char_flags1['\''] |= CHAR1_IS_QUOTE;
+  pp_char_flags1['\"'] |= CHAR1_IS_QUOTE;
+  pp_char_flags1['\\'] |= CHAR1_IS_BACKSLASH;
+
+  pp_char_flags2[ 0 ]  |= CHAR2_IS_NUL;
+  pp_char_flags2['\n'] |= CHAR2_IS_NEWLINE;
+  pp_char_flags2['\r'] |= CHAR2_IS_NEWLINE;
+  pp_char_flags2['/' ] |= CHAR2_IS_SLASH;
+}
+
+void init_preprocessor(void)
+{
+#ifdef HASHTEST
+  test_stringhash();
+#endif
+
+  init_preprocessor_scanner();
+  init_char_type();
+  //test_integer_constant_expression();
+
+  init_directive("define", HASH_define);
+  init_directive("else", HASH_else);
+  init_directive("elif", HASH_elif);
+  init_directive("endif", HASH_endif);
+  init_directive("error", HASH_error);
+  init_directive("if", HASH_if);
+  init_directive("ifdef", HASH_ifdef);
+  init_directive("ifndef", HASH_ifndef);
+  init_directive("include", HASH_include);
+  init_directive("line", HASH_line);
+  init_directive("pragma", HASH_pragma);
+  init_directive("undef", HASH_undef);
+}
+
+// File reader
+
+typedef struct
+{
+  char  * key; // resolved path
+  char  * contents;
+  int     contents_in_use;  // when multiprocessing or recursing, don't discard contents
+  size_t  contents_size;
+  int     contents_allocation_mode;
+} cached_header;
+
+typedef struct
+{
+  char  * key;
+  char  * value;
+} resolved_header;
+
+static cached_header   *cached_header_map;
+static size_t           cached_header_total_size;
+static resolved_header *resolved_header_map;
+static resolved_header *resolved_system_header_map;
+static stb_arena        resolve_arena;
+
+static size_t max_cached_headers = sizeof(size_t) > 4
+                                             ? (128U << 20) // 128MB in 64-bit
+                                             : (32U << 20); //  32MB in 32-bit
+
+void preprocessor_header_cache_set_size(size_t size)
+{
+  max_cached_headers = size;
+}
+
+void preprocessor_header_cache_flush(void)
+{
+  int i;
+  for (i=(int) hmlen(cached_header_map); i >= 0; ) {
+    if (cached_header_map[i].contents_in_use)
+      --i; // can't delete it
+    else {
+      free(cached_header_map[i].contents);
+      cached_header_map[i].contents = 0;
+    }
+  }
+}
+
+void init_preprocessor_cache(void)
+{
+  sh_new_arena(resolved_header_map);
+  sh_new_arena(resolved_system_header_map);
+}
+
+static int find_file_path(char path[4096], char *filename, char **path_list, ptrdiff_t path_list_len)
+{
+  int i;
+  FILE *f = NULL;
+
+  for (i=0; i < path_list_len; ++i) {
+    sprintf(path, "%s/%s", path_list[i], filename);
+    f = fopen(path, "rb");
+    if (f) {
+      fclose(f);
+      return 1;
+    }
+  }
+
+  return 0;
+}
+
+static char *my_strcpy(char *dest, char *src)
+{
+  while (*src)
+    *dest++ = *src++;
+  *dest = 0;
+  return dest;
+}
+
+static char *resolve_file(char path[4096], parse_state *ps, ppbool check_source_dirs, ppbool check_normal_includes, ppbool check_system_includes, char *filename, ptrdiff_t filename_len)
+{
+  pp_context *c = ps->context;
+
+  char filenamez[4096];
+  strncpy(filenamez, filename, filename_len);
+  filenamez[filename_len] = 0;
+
+  if (check_source_dirs) {
+    do {
+      FILE *f;
+      char *end;
+      if (ps->parent != NULL) {
+        end = my_strcpy(path, ps->filename);
+        while (end > path && end[-1] != '/' && end[-1] != '\\')
+          --end;
+        strcpy(end, filenamez);
+        f = fopen(path, "rb");
+        if (f) {
+          fclose(f);
+          return path;
+        }
+      }
+      ps = ps->parent;
+    } while (ps);
+  }
+
+  if (check_normal_includes) {
+    char *resolved = shget(resolved_header_map, filenamez);
+    if (resolved)
+      return resolved;
+
+    if (find_file_path(path, filenamez, c->include_paths, c->num_include_paths)) {
+      resolved = stb_arena_alloc_string(&resolve_arena, path);
+      shput(resolved_header_map, filenamez, resolved);
+      return resolved;
+    }
+  }
+
+  if (check_system_includes) {
+    char *resolved = shget(resolved_system_header_map, filenamez);
+    if (resolved)
+      return resolved;
+
+    if (find_file_path(path, filenamez, c->sys_include_paths, c->num_sys_include_paths)) {
+      resolved = stb_arena_alloc_string(&resolve_arena, path);
+      shput(resolved_system_header_map, filenamez, resolved);
+      return resolved;
+    }
+  }
+
+  return NULL;
+}
+
+#define LOAD_FILE_MODE_malloced         -2
+#define LOAD_FILE_MODE_memmappend       -1
+#define LOAD_FILE_MODE_multi_memmaped    0   // actual positive value is number of 4K pages past end of pointer that second mapping occurs
+
+static char *load_file_with_nul_terminator(char *filename, size_t *ptr_len, int *mode)
+{
+  char *p;
+  size_t len;
+  FILE *f = fopen(filename, "rb");
+  *mode = LOAD_FILE_MODE_malloced;
+  if (f == 0)
+    return NULL;
+  fseek(f, 0, SEEK_END);
+  len = ftell(f);
+  fseek(f, 0, SEEK_SET);
+  p = malloc(len+1);
+  p[len] = 0;
+  fread(p, 1, len, f);
+  fclose(f);
+  if (ptr_len)
+    *ptr_len = len+1;
+  return p;
+}
+
+static void load_file_free(char *filedata, int mode)
+{
+  free(filedata);
+}
+
+#if 0
+#include <windows.h>
+void *stb_mmap(char *filename, int *n)
+{
+   HANDLE f = CreateFile(filename, GENERIC_READ, FILE_SHARE_READ,  NULL, OPEN_EXISTING, FILE_ATTRIBUTE_NORMAL, NULL);
+   HANDLE m;
+   void *p;
+   if (!f) stb_fatal("Couldn't CreateFile() on '%s'", filename);
+   if (n) *n = GetFileSize(f, NULL);
+   m = CreateFileMapping(f, NULL, PAGE_READONLY, 0,0, NULL);
+   if (!m) stb_fatal("Couldn't CreateFileMapping() on '%s'", filename);
+   p = MapViewOfFile(m, FILE_MAP_READ, 0,0,0);
+   if (!p) stb_fatal("Couldn't MapViewOfFile() on '%s'", filename);
+   return p;
+}
+#endif
diff --git a/Engine/Source/Developer/ShaderPreprocessor/Private/stb_preprocess/preprocessor.h b/Engine/Source/Developer/ShaderPreprocessor/Private/stb_preprocess/preprocessor.h
new file mode 100644
--- /dev/null
+++ b/Engine/Source/Developer/ShaderPreprocessor/Private/stb_preprocess/preprocessor.h
@@ -0,0 +1,59 @@
+// Copyright Epic Games, Inc. All Rights Reserved.
+
+enum
+{
+  PP_RESULT_ok,
+  PP_RESULT_supplementary, // extra information associated with previous warning/error
+  PP_RESULT_undef_of_undefined_macro, // by default this is not an error
+
+  PP_RESULT_internal_error_out_of_memory,
+
+  PP_RESULT_ERROR=4,
+
+  PP_RESULT_counter_overflowed,
+  PP_RESULT_too_many_arguments_to_macro,
+  PP_RESULT_invalid_char,
+  PP_RESULT_unrecognized_directive,
+  PP_RESULT_directive_not_at_start_of_line,
+  PP_RESULT_unexpected_character_in_directive,
+  PP_RESULT_unexpected_character_after_end_of_directive,
+  PP_RESULT_undef_of_predefined_macro,
+  PP_RESULT_identifier_too_long,
+
+  PP_RESULT_count
+};
+
+enum
+{
+  PP_RESULT_MODE_error,          // default error type is 0 to simplify initialization
+  PP_RESULT_MODE_warning,
+  PP_RESULT_MODE_supplementary,
+
+  PP_RESULT_MODE_warning_fast,
+  PP_RESULT_MODE_no_warning
+};
+
+typedef struct macro_definition macro_definition;
+
+typedef struct
+{
+  char *filename;
+  int   line_number;
+  int   column;
+} pp_where;
+
+typedef struct
+{
+  char *message;
+  int   diagnostic_code;
+  int   error_level;
+  pp_where *where;
+} pp_diagnostic;
+
+extern void preprocessor_test(void);
+extern void init_preprocessor(void);
+extern char * preprocess_file(char * output_storage, char *filename, char **include_paths, int num_include_paths, char **sys_include_paths, int num_sys_include_paths, macro_definition **predefined_macros, int num_predefined_macros, pp_diagnostic **pd, int *num_pd);
+extern void preprocessor_file_free(char *text, pp_diagnostic *pd);
+typedef struct stb_arena stb_arena;
+extern macro_definition *pp_define(stb_arena *a, char *p);
+extern unsigned int preprocessor_hash(char *data, size_t len); // hash function for string of length 'len', roughly xxhash32
diff --git a/Engine/Source/Developer/ShaderPreprocessor/Private/stb_preprocess/stb_alloc.h b/Engine/Source/Developer/ShaderPreprocessor/Private/stb_preprocess/stb_alloc.h
new file mode 100644
--- /dev/null
+++ b/Engine/Source/Developer/ShaderPreprocessor/Private/stb_preprocess/stb_alloc.h
@@ -0,0 +1,142 @@
+// Copyright Epic Games, Inc. All Rights Reserved.
+
+#ifndef INCLUDE_STB_ALLOC_H
+#define INCLUDE_STB_ALLOC_H
+
+// @TODO rewrite to use a linked list of blocks instead of an array
+
+#include <stdlib.h>
+
+#ifdef __cplusplus
+#define STB_ALLOC_DEF   extern "C"
+#else
+#define STB_ALLOC_DEF   extern
+#endif
+
+typedef struct stb_arena stb_arena;
+
+// allocate from an arena
+STB_ALLOC_DEF void *stb_arena_alloc        (stb_arena *a, size_t size);
+STB_ALLOC_DEF void *stb_arena_alloc_aligned(stb_arena *a, size_t size, size_t align);
+
+// allocate a string from an arena
+STB_ALLOC_DEF char *stb_arena_alloc_string       (stb_arena *a, char *str);
+STB_ALLOC_DEF char *stb_arena_alloc_string_length(stb_arena *a, char *str, size_t length);
+
+// free the entire arena, leaves arena reset so it can be used again
+STB_ALLOC_DEF void stb_arena_free(stb_arena *a);
+
+// You allocate this structure yourself; initialize to all 0s and
+// don't mess with the innards. Optionally, you can initialize the first
+// field to the default block size you want, e.g.
+//     stb_arena arena = { 65536 };
+struct stb_arena
+{
+  size_t default_block_size;
+  size_t block_size;
+  void **blocks;
+  unsigned char *last_block;
+  int num_blocks;
+  size_t last_block_alloc_offset;
+  size_t last_block_size;
+};
+
+#endif // INCLUDE_STB_ALLOC_H
+
+#ifdef STB_ALLOC_IMPLEMENTATION
+
+void *stb_arena_alloc_aligned(stb_arena *a, size_t size, size_t align)
+{
+  void *p;
+  a->last_block_alloc_offset = (a->last_block_alloc_offset + align-1) & ~(align - 1);
+  if (a->last_block_alloc_offset + size > a->last_block_size) {
+    // initialize the block size if it's uninitialized
+    if (a->default_block_size == 0)
+      a->default_block_size = 8192;
+    if (a->block_size == 0)
+      a->block_size = a->default_block_size;
+
+    // increase the block size every 4 allocations
+    if ((a->num_blocks & 3) == 3)
+      a->block_size *= 2;
+
+    // the most recent block may be larger if the allocation itself is too large
+    a->last_block_size = a->block_size;
+    if (a->last_block_size < size)
+      a->last_block_size = size;
+
+    // grow the list of blocks if necessary
+    if (a->num_blocks == 0) {
+      a->blocks = malloc(sizeof(a->blocks[0]) * 32);
+    } else if (a->num_blocks >= 32) {
+      // if number of blocks is a power of two
+      if ((a->num_blocks & (a->num_blocks-1)) == 0) {
+        // grow table
+        a->blocks = realloc(a->blocks, sizeof(a->blocks[0]) * a->num_blocks * 2);
+      }
+    }
+
+    // allocate and record allocation
+    a->last_block = (unsigned char *) malloc(a->last_block_size+31);
+    a->blocks[a->num_blocks++] = a->last_block;
+    a->last_block_alloc_offset = 0;
+
+    {
+      // align allocation to multiple of 32 for future-proofing
+      union {
+        void   *p;
+        size_t  val;
+      } convert;
+      convert.p = a->last_block;
+      if (convert.val & 31) {
+        a->last_block += 32 - (convert.val & 31);
+      }
+    }
+  }
+
+  p = a->last_block + a->last_block_alloc_offset;
+  a->last_block_alloc_offset += size;
+  return p;
+}
+
+static size_t stb_arena_align_size(size_t size)
+{
+  if ((size & 7) != 0) return 4;
+  if (size >= 16) {
+    if (size >= 32 && (size & 31) == 0) return 32;
+    if (              (size & 15) == 0) return 16;
+  }
+  return 8;
+}
+
+void *stb_arena_alloc(stb_arena *a, size_t size)
+{
+  return stb_arena_alloc_aligned(a, size, stb_arena_align_size(size));
+}
+
+char *stb_arena_alloc_string_length(stb_arena *a, char *str, size_t length)
+{
+  char *p = (char *) stb_arena_alloc_aligned(a, length+1, 4);
+  memcpy(p, str, length);
+  p[length] = 0;
+  return p;
+}
+
+char *stb_arena_alloc_string(stb_arena *a, char *str)
+{
+  return stb_arena_alloc_string_length(a, str, strlen(str));
+}
+
+void stb_arena_free(stb_arena *a)
+{
+  int i;
+  for (i=0; i < a->num_blocks; ++i)
+    free(a->blocks[i]);
+  free(a->blocks);
+  a->blocks = 0;
+  a->num_blocks = 0;
+  a->last_block_size = 0;
+  a->last_block = 0;
+  a->block_size = 0;
+}
+#endif // STB_ALLOC_IMPLEMENTATION
diff --git a/Engine/Source/Developer/ShaderPreprocessor/Private/stb_preprocess/stb_ds.h b/Engine/Source/Developer/ShaderPreprocessor/Private/stb_preprocess/stb_ds.h
new file mode 100644
--- /dev/null
+++ b/Engine/Source/Developer/ShaderPreprocessor/Private/stb_preprocess/stb_ds.h
@@ -0,0 +1,1886 @@
+// Copyright Epic Games, Inc. All Rights Reserved.
+
+/* stb_ds.h - v0.65 - public domain data structures - Sean Barrett 2019
+
+   This is a single-header-file library that provides easy-to-use
+   dynamic arrays and hash tables for C (also works in C++).
+
+   For a gentle introduction:
+      http://nothings.org/stb_ds
+
+   To use this library, do this in *one* C or C++ file:
+      #define STB_DS_IMPLEMENTATION
+      #include "stb_ds.h"
+
+TABLE OF CONTENTS
+
+  Table of Contents
+  Compile-time options
+  License
+  Documentation
+  Notes
+  Notes - Dynamic arrays
+  Notes - Hash maps
+  Credits
+
+COMPILE-TIME OPTIONS
+
+  #define STBDS_NO_SHORT_NAMES
+
+     This flag needs to be set globally.
+
+     By default stb_ds exposes shorter function names that are not qualified
+     with the "stbds_" prefix. If these names conflict with the names in your
+     code, define this flag.
+
+  #define STBDS_SIPHASH_2_4
+
+     This flag only needs to be set in the file containing #define STB_DS_IMPLEMENTATION.
+
+     By default stb_ds.h hashes using a weaker variant of SipHash and a custom hash for
+     4- and 8-byte keys. On 64-bit platforms, you can define the above flag to force
+     stb_ds.h to use specification-compliant SipHash-2-4 for all keys. Doing so makes
+     hash table insertion about 20% slower on 4- and 8-byte keys, 5% slower on
+     64-byte keys, and 10% slower on 256-byte keys on my test computer.
+
+  #define STBDS_REALLOC(context,ptr,size) better_realloc
+  #define STBDS_FREE(context,ptr)         better_free
+
+     These defines only need to be set in the file containing #define STB_DS_IMPLEMENTATION.
+
+     By default stb_ds uses stdlib realloc() and free() for memory management. You can
+     substitute your own functions instead by defining these symbols. You must either
+     define both, or neither. Note that at the moment, 'context' will always be NULL.
+     @TODO add an array/hash initialization function that takes a memory context pointer.
+
+  #define STBDS_UNIT_TESTS
+
+     Defines a function stbds_unit_tests() that checks the functioning of the data structures.
+
+  Note that on older versions of gcc (e.g. 5.x.x) you may need to build with '-std=c++0x'
+     (or equivalentally '-std=c++11') when using anonymous structures as seen on the web
+     page or in STBDS_UNIT_TESTS.
+
+LICENSE
+
+  Placed in the public domain and also MIT licensed.
+  See end of file for detailed license information.
+
+DOCUMENTATION
+
+  Dynamic Arrays
+
+    Non-function interface:
+
+      Declare an empty dynamic array of type T
+        T* foo = NULL;
+
+      Access the i'th item of a dynamic array 'foo' of type T, T* foo:
+        foo[i]
+
+    Functions (actually macros)
+
+      arrfree:
+        void arrfree(T*);
+          Frees the array.
+
+      arrlen:
+        ptrdiff_t arrlen(T*);
+          Returns the number of elements in the array.
+
+      arrlenu:
+        size_t arrlenu(T*);
+          Returns the number of elements in the array as an unsigned type.
+
+      arrpop:
+        T arrpop(T* a)
+          Removes the final element of the array and returns it.
+
+      arrput:
+        T arrput(T* a, T b);
+          Appends the item b to the end of array a. Returns b.
+
+      arrins:
+        T arrins(T* a, int p, T b);
+          Inserts the item b into the middle of array a, into a[p],
+          moving the rest of the array over. Returns b.
+
+      arrinsn:
+        void arrins(T* a, int p, int n);
+          Inserts n uninitialized items into array a starting at a[p],
+          moving the rest of the array over.
+
+      arraddnptr:
+        T* arraddnptr(T* a, int n)
+          Appends n uninitialized items onto array at the end.
+          Returns a pointer to the first uninitialized item added.
+
+      arraddnindex:
+        size_t arraddnindex(T* a, int n)
+          Appends n uninitialized items onto array at the end.
+          Returns the index of the first uninitialized item added.
+
+      arrdel:
+        void arrdel(T* a, int p);
+          Deletes the element at a[p], moving the rest of the array over.
+
+      arrdeln:
+        void arrdel(T* a, int p, int n);
+          Deletes n elements starting at a[p], moving the rest of the array over.
+
+      arrdelswap:
+        void arrdelswap(T* a, int p);
+          Deletes the element at a[p], replacing it with the element from
+          the end of the array. O(1) performance.
+
+      arrsetlen:
+        void arrsetlen(T* a, int n);
+          Changes the length of the array to n. Allocates uninitialized
+          slots at the end if necessary.
+
+      arrsetcap:
+        size_t arrsetcap(T* a, int n);
+          Sets the length of allocated storage to at least n. It will not
+          change the length of the array.
+
+      arrcap:
+        size_t arrcap(T* a);
+          Returns the number of total elements the array can contain without
+          needing to be reallocated.
+
+  Hash maps & String hash maps
+
+    Given T is a structure type: struct { TK key; TV value; }. Note that some
+    functions do not require TV value and can have other fields. For string
+    hash maps, TK must be 'char *'.
+
+    Special interface:
+
+      stbds_rand_seed:
+        void stbds_rand_seed(size_t seed);
+          For security against adversarially chosen data, you should seed the
+          library with a strong random number. Or at least seed it with time().
+
+      stbds_hash_string:
+        size_t stbds_hash_string(char *str, size_t seed);
+          Returns a hash value for a string.
+
+      stbds_hash_bytes:
+        size_t stbds_hash_bytes(void *p, size_t len, size_t seed);
+          These functions hash an arbitrary number of bytes. The function
+          uses a custom hash for 4- and 8-byte data, and a weakened version
+          of SipHash for everything else. On 64-bit platforms you can get
+          specification-compliant SipHash-2-4 on all data by defining
+          STBDS_SIPHASH_2_4, at a significant cost in speed.
+
+    Non-function interface:
+
+      Declare an empty hash map of type T
+        T* foo = NULL;
+
+      Access the i'th entry in a hash table T* foo:
+        foo[i]
+
+    Function interface (actually macros):
+
+      hmfree
+      shfree
+        void hmfree(T*);
+        void shfree(T*);
+          Frees the hashmap and sets the pointer to NULL.
+
+      hmlen
+      shlen
+        ptrdiff_t hmlen(T*)
+        ptrdiff_t shlen(T*)
+          Returns the number of elements in the hashmap.
+
+      hmlenu
+      shlenu
+        size_t hmlenu(T*)
+        size_t shlenu(T*)
+          Returns the number of elements in the hashmap.
+
+      hmgeti
+      shgeti
+      hmgeti_ts
+        ptrdiff_t hmgeti(T*, TK key)
+        ptrdiff_t shgeti(T*, char* key)
+        ptrdiff_t hmgeti_ts(T*, TK key, ptrdiff_t tempvar)
+          Returns the index in the hashmap which has the key 'key', or -1
+          if the key is not present.
+
+      hmget
+      hmget_ts
+      shget
+        TV hmget(T*, TK key)
+        TV shget(T*, char* key)
+        TV hmget_ts(T*, TK key, ptrdiff_t tempvar)
+          Returns the value corresponding to 'key' in the hashmap.
+          The structure must have a 'value' field
+
+      hmgets
+      shgets
+        T hmgets(T*, TK key)
+        T shgets(T*, char* key)
+          Returns the structure corresponding to 'key' in the hashmap.
+
+      hmgetp
+      shgetp
+      hmgetp_ts
+      hmgetp_null
+      shgetp_null
+        T* hmgetp(T*, TK key)
+        T* shgetp(T*, char* key)
+        T* hmgetp_ts(T*, TK key, ptrdiff_t tempvar)
+        T* hmgetp_null(T*, TK key)
+        T* shgetp_null(T*, char *key)
+          Returns a pointer to the structure corresponding to 'key' in
+          the hashmap. Functions ending in "_null" return NULL if the key
+          is not present in the hashmap; the others return a pointer to a
+          structure holding the default value (but not the searched-for key).
+
+      hmdefault
+      shdefault
+        TV hmdefault(T*, TV value)
+        TV shdefault(T*, TV value)
+          Sets the default value for the hashmap, the value which will be
+          returned by hmget/shget if the key is not present.
+
+      hmdefaults
+      shdefaults
+        TV hmdefaults(T*, T item)
+        TV shdefaults(T*, T item)
+          Sets the default struct for the hashmap, the contents which will be
+          returned by hmgets/shgets if the key is not present.
+
+      hmput
+      shput
+        TV hmput(T*, TK key, TV value)
+        TV shput(T*, char* key, TV value)
+          Inserts a <key,value> pair into the hashmap. If the key is already
+          present in the hashmap, updates its value.
+
+      hmputs
+      shputs
+        T hmputs(T*, T item)
+        T shputs(T*, T item)
+          Inserts a struct with T.key into the hashmap. If the struct is already
+          present in the hashmap, updates it.
+
+      hmdel
+      shdel
+        int hmdel(T*, TK key)
+        int shdel(T*, char* key)
+          If 'key' is in the hashmap, deletes its entry and returns 1.
+          Otherwise returns 0.
+
+    Function interface (actually macros) for strings only:
+
+      sh_new_strdup
+        void sh_new_strdup(T*);
+          Overwrites the existing pointer with a newly allocated
+          string hashmap which will automatically allocate and free
+          each string key using realloc/free
+
+      sh_new_arena
+        void sh_new_arena(T*);
+          Overwrites the existing pointer with a newly allocated
+          string hashmap which will automatically allocate each string
+          key to a string arena. Every string key ever used by this
+          hash table remains in the arena until the arena is freed.
+          Additionally, any key which is deleted and reinserted will
+          be allocated multiple times in the string arena.
+
+NOTES
+
+  * These data structures are realloc'd when they grow, and the macro
+    "functions" write to the provided pointer. This means: (a) the pointer
+    must be an lvalue, and (b) the pointer to the data structure is not
+    stable, and you must maintain it the same as you would a realloc'd
+    pointer. For example, if you pass a pointer to a dynamic array to a
+    function which updates it, the function must return back the new
+    pointer to the caller. This is the price of trying to do this in C.
+
+  * The following are the only functions that are thread-safe on a single data
+    structure, i.e. can be run in multiple threads simultaneously on the same
+    data structure
+        hmlen        shlen
+        hmlenu       shlenu
+        hmget_ts     shget_ts
+        hmgeti_ts    shgeti_ts
+        hmgets_ts    shgets_ts
+
+  * You iterate over the contents of a dynamic array and a hashmap in exactly
+    the same way, using arrlen/hmlen/shlen:
+
+      for (i=0; i < arrlen(foo); ++i)
+         ... foo[i] ...
+
+  * All operations except arrins/arrdel are O(1) amortized, but individual
+    operations can be slow, so these data structures may not be suitable
+    for real time use. Dynamic arrays double in capacity as needed, so
+    elements are copied an average of once. Hash tables double/halve
+    their size as needed, with appropriate hysteresis to maintain O(1)
+    performance.
+
+NOTES - DYNAMIC ARRAY
+
+  * If you know how long a dynamic array is going to be in advance, you can avoid
+    extra memory allocations by using arrsetlen to allocate it to that length in
+    advance and use foo[n] while filling it out, or arrsetcap to allocate the memory
+    for that length and use arrput/arrpush as normal.
+
+  * Unlike some other versions of the dynamic array, this version should
+    be safe to use with strict-aliasing optimizations.
+
+NOTES - HASH MAP
+
+  * For compilers other than GCC and clang (e.g. Visual Studio), for hmput/hmget/hmdel
+    and variants, the key must be an lvalue (so the macro can take the address of it).
+    Extensions are used that eliminate this requirement if you're using C99 and later
+    in GCC or clang, or if you're using C++ in GCC. But note that this can make your
+    code less portable.
+
+  * To test for presence of a key in a hashmap, just do 'hmgeti(foo,key) >= 0'.
+
+  * The iteration order of your data in the hashmap is determined solely by the
+    order of insertions and deletions. In particular, if you never delete, new
+    keys are always added at the end of the array. This will be consistent
+    across all platforms and versions of the library. However, you should not
+    attempt to serialize the internal hash table, as the hash is not consistent
+    between different platforms, and may change with future versions of the library.
+
+  * Use sh_new_arena() for string hashmaps that you never delete from. Initialize
+    with NULL if you're managing the memory for your strings, or your strings are
+    never freed (at least until the hashmap is freed). Otherwise, use sh_new_strdup().
+    @TODO: make an arena variant that garbage collects the strings with a trivial
+    copy collector into a new arena whenever the table shrinks / rebuilds. Since
+    current arena recommendation is to only use arena if it never deletes, then
+    this can just replace current arena implementation.
+
+  * If adversarial input is a serious concern and you're on a 64-bit platform,
+    enable STBDS_SIPHASH_2_4 (see the 'Compile-time options' section), and pass
+    a strong random number to stbds_rand_seed.
+
+  * The default value for the hash table is stored in foo[-1], so if you
+    use code like 'hmget(T,k)->value = 5' you can accidentally overwrite
+    the value stored by hmdefault if 'k' is not present.
+
+CREDITS
+
+  Sean Barrett -- library, idea for dynamic array API/implementation
+  Per Vognsen  -- idea for hash table API/implementation
+  Rafael Sachetto -- arrpop()
+  github:HeroicKatora -- arraddn() reworking
+
+  Bugfixes:
+    Andy Durdin
+    Shane Liesegang
+    Vinh Truong
+    Andreas Molzer
+    github:hashitaku
+    github:srdjanstipic
+*/
+
+#ifdef STBDS_UNIT_TESTS
+#define _CRT_SECURE_NO_WARNINGS
+#endif
+
+#ifndef INCLUDE_STB_DS_H
+#define INCLUDE_STB_DS_H
+
+#include <stddef.h>
+#include <string.h>
+
+#ifndef STBDS_NO_SHORT_NAMES
+#define arrlen      stbds_arrlen
+#define arrlenu     stbds_arrlenu
+#define arrput      stbds_arrput
+#define arrpush     stbds_arrput
+#define arrpop      stbds_arrpop
+#define arrfree     stbds_arrfree
+#define arraddn     stbds_arraddn // deprecated, use one of the following instead:
+#define arraddnptr  stbds_arraddnptr
+#define arraddnindex stbds_arraddnindex
+#define arrsetlen   stbds_arrsetlen
+#define arrlast     stbds_arrlast
+#define arrins      stbds_arrins
+#define arrinsn     stbds_arrinsn
+#define arrdel      stbds_arrdel
+#define arrdeln     stbds_arrdeln
+#define arrdelswap  stbds_arrdelswap
+#define arrcap      stbds_arrcap
+#define arrsetcap   stbds_arrsetcap
+
+#define hmput       stbds_hmput
+#define hmputs      stbds_hmputs
+#define hmget       stbds_hmget
+#define hmget_ts    stbds_hmget_ts
+#define hmgets      stbds_hmgets
+#define hmgetp      stbds_hmgetp
+#define hmgetp_ts   stbds_hmgetp_ts
+#define hmgetp_null stbds_hmgetp_null
+#define hmgeti      stbds_hmgeti
+#define hmgeti_ts   stbds_hmgeti_ts
+#define hmdel       stbds_hmdel
+#define hmlen       stbds_hmlen
+#define hmlenu      stbds_hmlenu
+#define hmfree      stbds_hmfree
+#define hmdefault   stbds_hmdefault
+#define hmdefaults  stbds_hmdefaults
+
+#define shput       stbds_shput
+#define shputi      stbds_shputi
+#define shputs      stbds_shputs
+#define shget       stbds_shget
+#define shgeti      stbds_shgeti
+#define shgets      stbds_shgets
+#define shgetp      stbds_shgetp
+#define shgetp_null stbds_shgetp_null
+#define shdel       stbds_shdel
+#define shlen       stbds_shlen
+#define shlenu      stbds_shlenu
+#define shfree      stbds_shfree
+#define shdefault   stbds_shdefault
+#define shdefaults  stbds_shdefaults
+#define sh_new_arena  stbds_sh_new_arena
+#define sh_new_strdup stbds_sh_new_strdup
+
+#define stralloc    stbds_stralloc
+#define strreset    stbds_strreset
+#endif
+
+#if defined(STBDS_REALLOC) && !defined(STBDS_FREE) || !defined(STBDS_REALLOC) && defined(STBDS_FREE)
+#error "You must define both STBDS_REALLOC and STBDS_FREE, or neither."
+#endif
+#if !defined(STBDS_REALLOC) && !defined(STBDS_FREE)
+#include <stdlib.h>
+#define STBDS_REALLOC(c,p,s) realloc(p,s)
+#define STBDS_FREE(c,p)      free(p)
+#endif
+
+#ifdef _MSC_VER
+#define STBDS_NOTUSED(v)  (void)(v)
+#else
+#define STBDS_NOTUSED(v)  (void)sizeof(v)
+#endif
+
+#ifdef __cplusplus
+extern "C" {
+#endif
+
+// for security against attackers, seed the library with a random number, at least time() but stronger is better
+extern void stbds_rand_seed(size_t seed);
+
+// these are the hash functions used internally if you want to test them or use them for other purposes
+extern size_t stbds_hash_bytes(void *p, size_t len, size_t seed);
+extern size_t stbds_hash_string(char *str, size_t seed);
+
+// this is a simple string arena allocator, initialize with e.g. 'stbds_string_arena my_arena={0}'.
+typedef struct stbds_string_arena stbds_string_arena;
+extern char * stbds_stralloc(stbds_string_arena *a, char *str);
+extern void   stbds_strreset(stbds_string_arena *a);
+
+// have to #define STBDS_UNIT_TESTS to call this
+extern void stbds_unit_tests(void);
+
+///////////////
+//
+// Everything below here is implementation details
+//
+
+extern void * stbds_arrgrowf(void *a, size_t elemsize, size_t addlen, size_t min_cap);
+extern void   stbds_hmfree_func(void *p, size_t elemsize);
+extern void * stbds_hmget_key(void *a, size_t elemsize, void *key, size_t keysize, int mode);
+extern void * stbds_hmget_key_ts(void *a, size_t elemsize, void *key, size_t keysize, ptrdiff_t *temp, int mode);
+extern void * stbds_hmput_default(void *a, size_t elemsize);
+extern void * stbds_hmput_key(void *a, size_t elemsize, void *key, size_t keysize, int mode);
+extern void * stbds_hmdel_key(void *a, size_t elemsize, void *key, size_t keysize, size_t keyoffset, int mode);
+extern void * stbds_shmode_func(size_t elemsize, int mode);
+
+#ifdef __cplusplus
+}
+#endif
+
+#if defined(__GNUC__) || defined(__clang__)
+#define STBDS_HAS_TYPEOF
+#ifdef __cplusplus
+//#define STBDS_HAS_LITERAL_ARRAY  // this is currently broken for clang
+#endif
+#endif
+
+#if !defined(__cplusplus)
+#if defined(__STDC_VERSION__) && __STDC_VERSION__ >= 199901L
+#define STBDS_HAS_LITERAL_ARRAY
+#endif
+#endif
+
+// this macro takes the address of the argument, but on gcc/clang can accept rvalues
+#if defined(STBDS_HAS_LITERAL_ARRAY) && defined(STBDS_HAS_TYPEOF)
+  #if __clang__
+  #define STBDS_ADDRESSOF(typevar, value)     ((__typeof__(typevar)[1]){value}) // literal array decays to pointer to value
+  #else
+  #define STBDS_ADDRESSOF(typevar, value)     ((typeof(typevar)[1]){value}) // literal array decays to pointer to value
+  #endif
+#else
+#define STBDS_ADDRESSOF(typevar, value)     &(value)
+#endif
+
+#define STBDS_OFFSETOF(var,field)           ((char *) &(var)->field - (char *) (var))
+
+#define stbds_header(t)  ((stbds_array_header *) (t) - 1)
+#define stbds_temp(t)    stbds_header(t)->temp
+#define stbds_temp_key(t) (*(char **) stbds_header(t)->hash_table)
+
+#define stbds_arrsetcap(a,n)  (stbds_arrgrow(a,0,n))
+#define stbds_arrsetlen(a,n)  ((stbds_arrcap(a) < (size_t) (n) ? stbds_arrsetcap((a),(size_t)(n)),0 : 0), (a) ? stbds_header(a)->length = (size_t) (n) : 0)
+#define stbds_arrcap(a)       ((a) ? stbds_header(a)->capacity : 0)
+#define stbds_arrlen(a)       ((a) ? (ptrdiff_t) stbds_header(a)->length : 0)
+#define stbds_arrlenu(a)      ((a) ?             stbds_header(a)->length : 0)
+#define stbds_arrput(a,v)     (stbds_arrmaybegrow(a,1), (a)[stbds_header(a)->length++] = (v))
+#define stbds_arrpush         stbds_arrput  // synonym
+#define stbds_arrpop(a)       (stbds_header(a)->length--, (a)[stbds_header(a)->length])
+#define stbds_arraddn(a,n)    ((void)(stbds_arraddnoff(a, n)))    // deprecated, use one of the following instead:
+#define stbds_arraddnptr(a,n) (stbds_arrmaybegrow(a,n), stbds_header(a)->length += (n), &(a)[stbds_header(a)->length-(n)])
+#define stbds_arraddnoff(a,n) (stbds_arrmaybegrow(a,n), stbds_header(a)->length += (n), stbds_header(a)->length-(n))
+#define stbds_arraddnindex(a,n) stbds_arraddnoff(a,n)
+#define stbds_arrlast(a)      ((a)[stbds_header(a)->length-1])
+#define stbds_arrfree(a)      ((void) ((a) ? STBDS_FREE(NULL,stbds_header(a)) : (void)0), (a)=NULL)
+#define stbds_arrdel(a,i)     stbds_arrdeln(a,i,1)
+#define stbds_arrdeln(a,i,n)  (memmove(&(a)[i], &(a)[(i)+(n)], sizeof *(a) * (stbds_header(a)->length-(n)-(i))), stbds_header(a)->length -= (n))
+#define stbds_arrdelswap(a,i) ((a)[i] = stbds_arrlast(a), stbds_header(a)->length -= 1)
+#define stbds_arrinsn(a,i,n)  ((a)?stbds_temp(a)=(i):0,stbds_arraddn((a),(n)), memmove(&(a)[stbds_temp(a)+(n)], &(a)[stbds_temp(a)], sizeof *(a) * (stbds_header(a)->length-(n)-stbds_temp(a))))
+#define stbds_arrins(a,i,v)   (stbds_arrinsn((a),(i),1), (a)[i]=(v))
+
+#define stbds_arrmaybegrow(a,n)  ((!(a) || stbds_header(a)->length + (n) > stbds_header(a)->capacity) \
+                                  ? (stbds_arrgrow(a,n,0),0) : 0)
+
+#define stbds_arrgrow(a,b,c)   ((a) = stbds_arrgrowf_wrapper((a), sizeof *(a), (b), (c)))
+
+#define stbds_hmput(t, k, v) \
+    ((t) = stbds_hmput_key_wrapper((t), sizeof *(t), (void*) STBDS_ADDRESSOF((t)->key, (k)), sizeof (t)->key, 0),   \
+     (t)[stbds_temp((t)-1)].key = (k),    \
+     (t)[stbds_temp((t)-1)].value = (v))
+
+#define stbds_hmputs(t, s) \
+    ((t) = stbds_hmput_key_wrapper((t), sizeof *(t), &(s).key, sizeof (s).key, STBDS_HM_BINARY), \
+     (t)[stbds_temp((t)-1)] = (s))
+
+#define stbds_hmgeti(t,k) \
+    ((t) = stbds_hmget_key_wrapper((t), sizeof *(t), (void*) STBDS_ADDRESSOF((t)->key, (k)), sizeof (t)->key, STBDS_HM_BINARY), \
+      stbds_temp((t)-1))
+
+#define stbds_hmgeti_ts(t,k,temp) \
+    ((t) = stbds_hmget_key_ts_wrapper((t), sizeof *(t), (void*) STBDS_ADDRESSOF((t)->key, (k)), sizeof (t)->key, &(temp), STBDS_HM_BINARY), \
+      (temp))
+
+#define stbds_hmgetp(t, k) \
+    ((void) stbds_hmgeti(t,k), &(t)[stbds_temp((t)-1)])
+
+#define stbds_hmgetp_ts(t, k, temp) \
+    ((void) stbds_hmgeti_ts(t,k,temp), &(t)[temp])
+
+#define stbds_hmdel(t,k) \
+    (((t) = stbds_hmdel_key_wrapper((t),sizeof *(t), (void*) STBDS_ADDRESSOF((t)->key, (k)), sizeof (t)->key, STBDS_OFFSETOF((t),key), STBDS_HM_BINARY)),(t)?stbds_temp((t)-1):0)
+
+#define stbds_hmdefault(t, v) \
+    ((t) = stbds_hmput_default_wrapper((t), sizeof *(t)), (t)[-1].value = (v))
+
+#define stbds_hmdefaults(t, s) \
+    ((t) = stbds_hmput_default_wrapper((t), sizeof *(t)), (t)[-1] = (s))
+
+#define stbds_hmfree(p)        \
+    ((void) ((p) != NULL ? stbds_hmfree_func((p)-1,sizeof*(p)),0 : 0),(p)=NULL)
+
+#define stbds_hmgets(t, k)    (*stbds_hmgetp(t,k))
+#define stbds_hmget(t, k)     (stbds_hmgetp(t,k)->value)
+#define stbds_hmget_ts(t, k, temp)  (stbds_hmgetp_ts(t,k,temp)->value)
+#define stbds_hmlen(t)        ((t) ? (ptrdiff_t) stbds_header((t)-1)->length-1 : 0)
+#define stbds_hmlenu(t)       ((t) ?             stbds_header((t)-1)->length-1 : 0)
+#define stbds_hmgetp_null(t,k)  (stbds_hmgeti(t,k) == -1 ? NULL : &(t)[stbds_temp((t)-1)])
+
+#define stbds_shput(t, k, v) \
+    ((t) = stbds_hmput_key_wrapper((t), sizeof *(t), (void*) (k), sizeof (t)->key, STBDS_HM_STRING),   \
+     (t)[stbds_temp((t)-1)].value = (v))
+
+#define stbds_shputi(t, k, v) \
+    ((t) = stbds_hmput_key_wrapper((t), sizeof *(t), (void*) (k), sizeof (t)->key, STBDS_HM_STRING),   \
+     (t)[stbds_temp((t)-1)].value = (v), stbds_temp((t)-1))
+
+#define stbds_shputs(t, s) \
+    ((t) = stbds_hmput_key_wrapper((t), sizeof *(t), (void*) (s).key, sizeof (s).key, STBDS_HM_STRING), \
+     (t)[stbds_temp((t)-1)] = (s), \
+     (t)[stbds_temp((t)-1)].key = stbds_temp_key((t)-1)) // above line overwrites whole structure, so must rewrite key here if it was allocated internally
+
+#define stbds_pshput(t, p) \
+    ((t) = stbds_hmput_key_wrapper((t), sizeof *(t), (void*) (p)->key, sizeof (p)->key, STBDS_HM_PTR_TO_STRING), \
+     (t)[stbds_temp((t)-1)] = (p))
+
+#define stbds_shgeti(t,k) \
+     ((t) = stbds_hmget_key_wrapper((t), sizeof *(t), (void*) (k), sizeof (t)->key, STBDS_HM_STRING), \
+      stbds_temp((t)-1))
+
+#define stbds_pshgeti(t,k) \
+     ((t) = stbds_hmget_key_wrapper((t), sizeof *(t), (void*) (k), sizeof (*(t))->key, STBDS_HM_PTR_TO_STRING), \
+      stbds_temp((t)-1))
+
+#define stbds_shgetp(t, k) \
+    ((void) stbds_shgeti(t,k), &(t)[stbds_temp((t)-1)])
+
+#define stbds_pshget(t, k) \
+    ((void) stbds_pshgeti(t,k), (t)[stbds_temp((t)-1)])
+
+#define stbds_shdel(t,k) \
+    (((t) = stbds_hmdel_key_wrapper((t),sizeof *(t), (void*) (k), sizeof (t)->key, STBDS_OFFSETOF((t),key), STBDS_HM_STRING)),(t)?stbds_temp((t)-1):0)
+#define stbds_pshdel(t,k) \
+    (((t) = stbds_hmdel_key_wrapper((t),sizeof *(t), (void*) (k), sizeof (*(t))->key, STBDS_OFFSETOF(*(t),key), STBDS_HM_PTR_TO_STRING)),(t)?stbds_temp((t)-1):0)
+
+#define stbds_sh_new_arena(t)  \
+    ((t) = stbds_shmode_func_wrapper(t, sizeof *(t), STBDS_SH_ARENA))
+#define stbds_sh_new_strdup(t) \
+    ((t) = stbds_shmode_func_wrapper(t, sizeof *(t), STBDS_SH_STRDUP))
+
+#define stbds_shdefault(t, v)  stbds_hmdefault(t,v)
+#define stbds_shdefaults(t, s) stbds_hmdefaults(t,s)
+
+#define stbds_shfree       stbds_hmfree
+#define stbds_shlenu       stbds_hmlenu
+
+#define stbds_shgets(t, k) (*stbds_shgetp(t,k))
+#define stbds_shget(t, k)  (stbds_shgetp(t,k)->value)
+#define stbds_shgetp_null(t,k)  (stbds_shgeti(t,k) == -1 ? NULL : &(t)[stbds_temp((t)-1)])
+#define stbds_shlen        stbds_hmlen
+
+typedef struct
+{
+  size_t      length;
+  size_t      capacity;
+  void      * hash_table;
+  ptrdiff_t   temp;
+} stbds_array_header;
+
+typedef struct stbds_string_block
+{
+  struct stbds_string_block *next;
+  char storage[8];
+} stbds_string_block;
+
+struct stbds_string_arena
+{
+  stbds_string_block *storage;
+  size_t remaining;
+  unsigned char block;
+  unsigned char mode;  // this isn't used by the string arena itself
+};
+
+#define STBDS_HM_BINARY         0
+#define STBDS_HM_STRING         1
+
+enum
+{
+   STBDS_SH_NONE,
+   STBDS_SH_DEFAULT,
+   STBDS_SH_STRDUP,
+   STBDS_SH_ARENA
+};
+
+#ifdef __cplusplus
+// in C we use implicit assignment from these void*-returning functions to T*.
+// in C++ these templates make the same code work
+template<class T> static T * stbds_arrgrowf_wrapper(T *a, size_t elemsize, size_t addlen, size_t min_cap) {
+  return (T*)stbds_arrgrowf((void *)a, elemsize, addlen, min_cap);
+}
+template<class T> static T * stbds_hmget_key_wrapper(T *a, size_t elemsize, void *key, size_t keysize, int mode) {
+  return (T*)stbds_hmget_key((void*)a, elemsize, key, keysize, mode);
+}
+template<class T> static T * stbds_hmget_key_ts_wrapper(T *a, size_t elemsize, void *key, size_t keysize, ptrdiff_t *temp, int mode) {
+  return (T*)stbds_hmget_key_ts((void*)a, elemsize, key, keysize, temp, mode);
+}
+template<class T> static T * stbds_hmput_default_wrapper(T *a, size_t elemsize) {
+  return (T*)stbds_hmput_default((void *)a, elemsize);
+}
+template<class T> static T * stbds_hmput_key_wrapper(T *a, size_t elemsize, void *key, size_t keysize, int mode) {
+  return (T*)stbds_hmput_key((void*)a, elemsize, key, keysize, mode);
+}
+template<class T> static T * stbds_hmdel_key_wrapper(T *a, size_t elemsize, void *key, size_t keysize, size_t keyoffset, int mode){
+  return (T*)stbds_hmdel_key((void*)a, elemsize, key, keysize, keyoffset, mode);
+}
+template<class T> static T * stbds_shmode_func_wrapper(T *, size_t elemsize, int mode) {
+  return (T*)stbds_shmode_func(elemsize, mode);
+}
+#else
+#define stbds_arrgrowf_wrapper            stbds_arrgrowf
+#define stbds_hmget_key_wrapper           stbds_hmget_key
+#define stbds_hmget_key_ts_wrapper        stbds_hmget_key_ts
+#define stbds_hmput_default_wrapper       stbds_hmput_default
+#define stbds_hmput_key_wrapper           stbds_hmput_key
+#define stbds_hmdel_key_wrapper           stbds_hmdel_key
+#define stbds_shmode_func_wrapper(t,e,m)  stbds_shmode_func(e,m)
+#endif
+
+#endif // INCLUDE_STB_DS_H
+
+
+//////////////////////////////////////////////////////////////////////////////
+//
+//   IMPLEMENTATION
+//
+
+#ifdef STB_DS_IMPLEMENTATION
+#include <assert.h>
+#include <string.h>
+
+#ifndef STBDS_ASSERT
+#define STBDS_ASSERT_WAS_UNDEFINED
+#define STBDS_ASSERT(x)   ((void) 0)
+#endif
+
+#ifdef STBDS_STATISTICS
+#define STBDS_STATS(x)   x
+size_t stbds_array_grow;
+size_t stbds_hash_grow;
+size_t stbds_hash_shrink;
+size_t stbds_hash_rebuild;
+size_t stbds_hash_probes;
+size_t stbds_hash_alloc;
+size_t stbds_rehash_probes;
+size_t stbds_rehash_items;
+#else
+#define STBDS_STATS(x)
+#endif
+
+//
+// stbds_arr implementation
+//
+
+//int *prev_allocs[65536];
+//int num_prev;
+
+void *stbds_arrgrowf(void *a, size_t elemsize, size_t addlen, size_t min_cap)
+{
+  stbds_array_header temp={0}; // force debugging
+  void *b;
+  size_t min_len = stbds_arrlen(a) + addlen;
+  (void) sizeof(temp);
+
+  // compute the minimum capacity needed
+  if (min_len > min_cap)
+    min_cap = min_len;
+
+  if (min_cap <= stbds_arrcap(a))
+    return a;
+
+  // increase needed capacity to guarantee O(1) amortized
+  if (min_cap < 2 * stbds_arrcap(a))
+    min_cap = 2 * stbds_arrcap(a);
+  else if (min_cap < 4)
+    min_cap = 4;
+
+  //if (num_prev < 65536) if (a) prev_allocs[num_prev++] = (int *) ((char *) a+1);
+  //if (num_prev == 2201)
+  //  num_prev = num_prev;
+  b = STBDS_REALLOC(NULL, (a) ? stbds_header(a) : 0, elemsize * min_cap + sizeof(stbds_array_header));
+  //if (num_prev < 65536) prev_allocs[num_prev++] = (int *) (char *) b;
+  b = (char *) b + sizeof(stbds_array_header);
+  if (a == NULL) {
+    stbds_header(b)->length = 0;
+    stbds_header(b)->hash_table = 0;
+    stbds_header(b)->temp = 0;
+  } else {
+    STBDS_STATS(++stbds_array_grow);
+  }
+  stbds_header(b)->capacity = min_cap;
+
+  return b;
+}
+
+//
+// stbds_hm hash table implementation
+//
+
+#ifdef STBDS_INTERNAL_SMALL_BUCKET
+#define STBDS_BUCKET_LENGTH      4
+#else
+#define STBDS_BUCKET_LENGTH      8
+#endif
+
+#define STBDS_BUCKET_SHIFT      (STBDS_BUCKET_LENGTH == 8 ? 3 : 2)
+#define STBDS_BUCKET_MASK       (STBDS_BUCKET_LENGTH-1)
+#define STBDS_CACHE_LINE_SIZE   64
+
+#define STBDS_ALIGN_FWD(n,a)   (((n) + (a) - 1) & ~((a)-1))
+
+typedef struct
+{
+   size_t    hash [STBDS_BUCKET_LENGTH];
+   ptrdiff_t index[STBDS_BUCKET_LENGTH];
+} stbds_hash_bucket; // in 32-bit, this is one 64-byte cache line; in 64-bit, each array is one 64-byte cache line
+
+typedef struct
+{
+  char * temp_key; // this MUST be the first field of the hash table
+  size_t slot_count;
+  size_t used_count;
+  size_t used_count_threshold;
+  size_t used_count_shrink_threshold;
+  size_t tombstone_count;
+  size_t tombstone_count_threshold;
+  size_t seed;
+  size_t slot_count_log2;
+  stbds_string_arena string;
+  stbds_hash_bucket *storage; // not a separate allocation, just 64-byte aligned storage after this struct
+} stbds_hash_index;
+
+#define STBDS_INDEX_EMPTY    -1
+#define STBDS_INDEX_DELETED  -2
+#define STBDS_INDEX_IN_USE(x)  ((x) >= 0)
+
+#define STBDS_HASH_EMPTY      0
+#define STBDS_HASH_DELETED    1
+
+static size_t stbds_hash_seed=0x31415926;
+
+void stbds_rand_seed(size_t seed)
+{
+  stbds_hash_seed = seed;
+}
+
+#define stbds_load_32_or_64(var, temp, v32, v64_hi, v64_lo)                                          \
+  temp = v64_lo ^ v32, temp <<= 16, temp <<= 16, temp >>= 16, temp >>= 16, /* discard if 32-bit */   \
+  var = v64_hi, var <<= 16, var <<= 16,                                    /* discard if 32-bit */   \
+  var ^= temp ^ v32
+
+#define STBDS_SIZE_T_BITS           ((sizeof (size_t)) * 8)
+
+static size_t stbds_probe_position(size_t hash, size_t slot_count, size_t slot_log2)
+{
+  size_t pos;
+  STBDS_NOTUSED(slot_log2);
+  pos = hash & (slot_count-1);
+  #ifdef STBDS_INTERNAL_BUCKET_START
+  pos &= ~STBDS_BUCKET_MASK;
+  #endif
+  return pos;
+}
+
+static size_t stbds_log2(size_t slot_count)
+{
+  size_t n=0;
+  while (slot_count > 1) {
+    slot_count >>= 1;
+    ++n;
+  }
+  return n;
+}
+
+static stbds_hash_index *stbds_make_hash_index(size_t slot_count, stbds_hash_index *ot)
+{
+  stbds_hash_index *t;
+  t = (stbds_hash_index *) STBDS_REALLOC(NULL,0,(slot_count >> STBDS_BUCKET_SHIFT) * sizeof(stbds_hash_bucket) + sizeof(stbds_hash_index) + STBDS_CACHE_LINE_SIZE-1);
+  t->storage = (stbds_hash_bucket *) STBDS_ALIGN_FWD((size_t) (t+1), STBDS_CACHE_LINE_SIZE);
+  t->slot_count = slot_count;
+  t->slot_count_log2 = stbds_log2(slot_count);
+  t->tombstone_count = 0;
+  t->used_count = 0;
+
+  #if 0 // A1
+  t->used_count_threshold        = slot_count*12/16; // if 12/16th of table is occupied, grow
+  t->tombstone_count_threshold   = slot_count* 2/16; // if tombstones are 2/16th of table, rebuild
+  t->used_count_shrink_threshold = slot_count* 4/16; // if table is only 4/16th full, shrink
+  #elif 1 // A2
+  //t->used_count_threshold        = slot_count*12/16; // if 12/16th of table is occupied, grow
+  //t->tombstone_count_threshold   = slot_count* 3/16; // if tombstones are 3/16th of table, rebuild
+  //t->used_count_shrink_threshold = slot_count* 4/16; // if table is only 4/16th full, shrink
+
+  // compute without overflowing
+  t->used_count_threshold        = slot_count - (slot_count>>2);
+  t->tombstone_count_threshold   = (slot_count>>3) + (slot_count>>4);
+  t->used_count_shrink_threshold = slot_count >> 2;
+
+  #elif 0 // B1
+  t->used_count_threshold        = slot_count*13/16; // if 13/16th of table is occupied, grow
+  t->tombstone_count_threshold   = slot_count* 2/16; // if tombstones are 2/16th of table, rebuild
+  t->used_count_shrink_threshold = slot_count* 5/16; // if table is only 5/16th full, shrink
+  #else // C1
+  t->used_count_threshold        = slot_count*14/16; // if 14/16th of table is occupied, grow
+  t->tombstone_count_threshold   = slot_count* 2/16; // if tombstones are 2/16th of table, rebuild
+  t->used_count_shrink_threshold = slot_count* 6/16; // if table is only 6/16th full, shrink
+  #endif
+  // Following statistics were measured on a Core i7-6700 @ 4.00Ghz, compiled with clang 7.0.1 -O2
+    // Note that the larger tables have high variance as they were run fewer times
+  //     A1            A2          B1           C1
+  //    0.10ms :     0.10ms :     0.10ms :     0.11ms :      2,000 inserts creating 2K table
+  //    0.96ms :     0.95ms :     0.97ms :     1.04ms :     20,000 inserts creating 20K table
+  //   14.48ms :    14.46ms :    10.63ms :    11.00ms :    200,000 inserts creating 200K table
+  //  195.74ms :   196.35ms :   203.69ms :   214.92ms :  2,000,000 inserts creating 2M table
+  // 2193.88ms :  2209.22ms :  2285.54ms :  2437.17ms : 20,000,000 inserts creating 20M table
+  //   65.27ms :    53.77ms :    65.33ms :    65.47ms : 500,000 inserts & deletes in 2K table
+  //   72.78ms :    62.45ms :    71.95ms :    72.85ms : 500,000 inserts & deletes in 20K table
+  //   89.47ms :    77.72ms :    96.49ms :    96.75ms : 500,000 inserts & deletes in 200K table
+  //   97.58ms :    98.14ms :    97.18ms :    97.53ms : 500,000 inserts & deletes in 2M table
+  //  118.61ms :   119.62ms :   120.16ms :   118.86ms : 500,000 inserts & deletes in 20M table
+  //  192.11ms :   194.39ms :   196.38ms :   195.73ms : 500,000 inserts & deletes in 200M table
+
+  if (slot_count <= STBDS_BUCKET_LENGTH)
+    t->used_count_shrink_threshold = 0;
+  // to avoid infinite loop, we need to guarantee that at least one slot is empty and will terminate probes
+  STBDS_ASSERT(t->used_count_threshold + t->tombstone_count_threshold < t->slot_count);
+  STBDS_STATS(++stbds_hash_alloc);
+  if (ot) {
+    t->string = ot->string;
+    // reuse old seed so we can reuse old hashes so below "copy out old data" doesn't do any hashing
+    t->seed = ot->seed;
+  } else {
+    size_t a,b,temp;
+    memset(&t->string, 0, sizeof(t->string));
+    t->seed = stbds_hash_seed;
+    // LCG
+    // in 32-bit, a =          2147001325   b =  715136305
+    // in 64-bit, a = 2862933555777941757   b = 3037000493
+    stbds_load_32_or_64(a,temp, 2147001325, 0x27bb2ee6, 0x87b0b0fd);
+    stbds_load_32_or_64(b,temp,  715136305,          0, 0xb504f32d);
+    stbds_hash_seed = stbds_hash_seed  * a + b;
+  }
+
+  {
+    size_t i,j;
+    for (i=0; i < slot_count >> STBDS_BUCKET_SHIFT; ++i) {
+      stbds_hash_bucket *b = &t->storage[i];
+      for (j=0; j < STBDS_BUCKET_LENGTH; ++j)
+        b->hash[j] = STBDS_HASH_EMPTY;
+      for (j=0; j < STBDS_BUCKET_LENGTH; ++j)
+        b->index[j] = STBDS_INDEX_EMPTY;
+    }
+  }
+
+  // copy out the old data, if any
+  if (ot) {
+    size_t i,j;
+    t->used_count = ot->used_count;
+    for (i=0; i < ot->slot_count >> STBDS_BUCKET_SHIFT; ++i) {
+      stbds_hash_bucket *ob = &ot->storage[i];
+      for (j=0; j < STBDS_BUCKET_LENGTH; ++j) {
+        if (STBDS_INDEX_IN_USE(ob->index[j])) {
+          size_t hash = ob->hash[j];
+          size_t pos = stbds_probe_position(hash, t->slot_count, t->slot_count_log2);
+          size_t step = STBDS_BUCKET_LENGTH;
+          STBDS_STATS(++stbds_rehash_items);
+          for (;;) {
+            size_t limit,z;
+            stbds_hash_bucket *bucket;
+            bucket = &t->storage[pos >> STBDS_BUCKET_SHIFT];
+            STBDS_STATS(++stbds_rehash_probes);
+
+            for (z=pos & STBDS_BUCKET_MASK; z < STBDS_BUCKET_LENGTH; ++z) {
+              if (bucket->hash[z] == 0) {
+                bucket->hash[z] = hash;
+                bucket->index[z] = ob->index[j];
+                goto done;
+              }
+            }
+
+            limit = pos & STBDS_BUCKET_MASK;
+            for (z = 0; z < limit; ++z) {
+              if (bucket->hash[z] == 0) {
+                bucket->hash[z] = hash;
+                bucket->index[z] = ob->index[j];
+                goto done;
+              }
+            }
+
+            pos += step;                  // quadratic probing
+            step += STBDS_BUCKET_LENGTH;
+            pos &= (t->slot_count-1);
+          }
+        }
+       done:
+        ;
+      }
+    }
+  }
+
+  return t;
+}
+
+#define STBDS_ROTATE_LEFT(val, n)   (((val) << (n)) | ((val) >> (STBDS_SIZE_T_BITS - (n))))
+#define STBDS_ROTATE_RIGHT(val, n)  (((val) >> (n)) | ((val) << (STBDS_SIZE_T_BITS - (n))))
+
+size_t stbds_hash_string(char *str, size_t seed)
+{
+  size_t hash = seed;
+  while (*str)
+     hash = STBDS_ROTATE_LEFT(hash, 9) + (unsigned char) *str++;
+
+  // Thomas Wang 64-to-32 bit mix function, hopefully also works in 32 bits
+  hash ^= seed;
+  hash = (~hash) + (hash << 18);
+  hash ^= hash ^ STBDS_ROTATE_RIGHT(hash,31);
+  hash = hash * 21;
+  hash ^= hash ^ STBDS_ROTATE_RIGHT(hash,11);
+  hash += (hash << 6);
+  hash ^= STBDS_ROTATE_RIGHT(hash,22);
+  return hash+seed;
+}
+
+#ifdef STBDS_SIPHASH_2_4
+#define STBDS_SIPHASH_C_ROUNDS 2
+#define STBDS_SIPHASH_D_ROUNDS 4
+typedef int STBDS_SIPHASH_2_4_can_only_be_used_in_64_bit_builds[sizeof(size_t) == 8 ? 1 : -1];
+#endif
+
+#ifndef STBDS_SIPHASH_C_ROUNDS
+#define STBDS_SIPHASH_C_ROUNDS 1
+#endif
+#ifndef STBDS_SIPHASH_D_ROUNDS
+#define STBDS_SIPHASH_D_ROUNDS 1
+#endif
+
+#ifdef _MSC_VER
+#pragma warning(push)
+#pragma warning(disable:4127) // conditional expression is constant, for do..while(0) and sizeof()==
+#endif
+
+static size_t stbds_siphash_bytes(void *p, size_t len, size_t seed)
+{
+  unsigned char *d = (unsigned char *) p;
+  size_t i,j;
+  size_t v0,v1,v2,v3, data;
+
+  // hash that works on 32- or 64-bit registers without knowing which we have
+  // (computes different results on 32-bit and 64-bit platform)
+  // derived from siphash, but on 32-bit platforms very different as it uses 4 32-bit state not 4 64-bit
+  v0 = ((((size_t) 0x736f6d65 << 16) << 16) + 0x70736575) ^  seed;
+  v1 = ((((size_t) 0x646f7261 << 16) << 16) + 0x6e646f6d) ^ ~seed;
+  v2 = ((((size_t) 0x6c796765 << 16) << 16) + 0x6e657261) ^  seed;
+  v3 = ((((size_t) 0x74656462 << 16) << 16) + 0x79746573) ^ ~seed;
+
+  #ifdef STBDS_TEST_SIPHASH_2_4
+  // hardcoded with key material in the siphash test vectors
+  v0 ^= 0x0706050403020100ull ^  seed;
+  v1 ^= 0x0f0e0d0c0b0a0908ull ^ ~seed;
+  v2 ^= 0x0706050403020100ull ^  seed;
+  v3 ^= 0x0f0e0d0c0b0a0908ull ^ ~seed;
+  #endif
+
+  #define STBDS_SIPROUND() \
+    do {                   \
+      v0 += v1; v1 = STBDS_ROTATE_LEFT(v1, 13);  v1 ^= v0; v0 = STBDS_ROTATE_LEFT(v0,STBDS_SIZE_T_BITS/2); \
+      v2 += v3; v3 = STBDS_ROTATE_LEFT(v3, 16);  v3 ^= v2;                                                 \
+      v2 += v1; v1 = STBDS_ROTATE_LEFT(v1, 17);  v1 ^= v2; v2 = STBDS_ROTATE_LEFT(v2,STBDS_SIZE_T_BITS/2); \
+      v0 += v3; v3 = STBDS_ROTATE_LEFT(v3, 21);  v3 ^= v0;                                                 \
+    } while (0)
+
+  for (i=0; i+sizeof(size_t) <= len; i += sizeof(size_t), d += sizeof(size_t)) {
+    data = d[0] | (d[1] << 8) | (d[2] << 16) | (d[3] << 24);
+    data |= (size_t) (d[4] | (d[5] << 8) | (d[6] << 16) | (d[7] << 24)) << 16 << 16; // discarded if size_t == 4
+
+    v3 ^= data;
+    for (j=0; j < STBDS_SIPHASH_C_ROUNDS; ++j)
+      STBDS_SIPROUND();
+    v0 ^= data;
+  }
+  data = len << (STBDS_SIZE_T_BITS-8);
+  switch (len - i) {
+    case 7: data |= ((size_t) d[6] << 24) << 24; // fall through
+    case 6: data |= ((size_t) d[5] << 20) << 20; // fall through
+    case 5: data |= ((size_t) d[4] << 16) << 16; // fall through
+    case 4: data |= (d[3] << 24); // fall through
+    case 3: data |= (d[2] << 16); // fall through
+    case 2: data |= (d[1] << 8); // fall through
+    case 1: data |= d[0]; // fall through
+    case 0: break;
+  }
+  v3 ^= data;
+  for (j=0; j < STBDS_SIPHASH_C_ROUNDS; ++j)
+    STBDS_SIPROUND();
+  v0 ^= data;
+  v2 ^= 0xff;
+  for (j=0; j < STBDS_SIPHASH_D_ROUNDS; ++j)
+    STBDS_SIPROUND();
+
+#ifdef STBDS_SIPHASH_2_4
+  return v0^v1^v2^v3;
+#else
+  return v1^v2^v3; // slightly stronger since v0^v3 in above cancels out final round operation? I tweeted at the authors of SipHash about this but they didn't reply
+#endif
+}
+
+size_t stbds_hash_bytes(void *p, size_t len, size_t seed)
+{
+#ifdef STBDS_SIPHASH_2_4
+  return stbds_siphash_bytes(p,len,seed);
+#else
+  unsigned char *d = (unsigned char *) p;
+
+  if (len == 4) {
+    unsigned int hash = d[0] | (d[1] << 8) | (d[2] << 16) | (d[3] << 24);
+    #if 0
+    // HASH32-A  Bob Jenkin's hash function w/o large constants
+    hash ^= seed;
+    hash -= (hash<<6);
+    hash ^= (hash>>17);
+    hash -= (hash<<9);
+    hash ^= seed;
+    hash ^= (hash<<4);
+    hash -= (hash<<3);
+    hash ^= (hash<<10);
+    hash ^= (hash>>15);
+    #elif 1
+    // HASH32-BB  Bob Jenkin's presumably-accidental version of Thomas Wang hash with rotates turned into shifts.
+    // Note that converting these back to rotates makes it run a lot slower, presumably due to collisions, so I'm
+    // not really sure what's going on.
+    hash ^= seed;
+    hash = (hash ^ 61) ^ (hash >> 16);
+    hash = hash + (hash << 3);
+    hash = hash ^ (hash >> 4);
+    hash = hash * 0x27d4eb2d;
+    hash ^= seed;
+    hash = hash ^ (hash >> 15);
+    #else  // HASH32-C   -  Murmur3
+    hash ^= seed;
+    hash *= 0xcc9e2d51;
+    hash = (hash << 17) | (hash >> 15);
+    hash *= 0x1b873593;
+    hash ^= seed;
+    hash = (hash << 19) | (hash >> 13);
+    hash = hash*5 + 0xe6546b64;
+    hash ^= hash >> 16;
+    hash *= 0x85ebca6b;
+    hash ^= seed;
+    hash ^= hash >> 13;
+    hash *= 0xc2b2ae35;
+    hash ^= hash >> 16;
+    #endif
+    // Following statistics were measured on a Core i7-6700 @ 4.00Ghz, compiled with clang 7.0.1 -O2
+    // Note that the larger tables have high variance as they were run fewer times
+    //  HASH32-A   //  HASH32-BB  //  HASH32-C
+    //    0.10ms   //    0.10ms   //    0.10ms :      2,000 inserts creating 2K table
+    //    0.96ms   //    0.95ms   //    0.99ms :     20,000 inserts creating 20K table
+    //   14.69ms   //   14.43ms   //   14.97ms :    200,000 inserts creating 200K table
+    //  199.99ms   //  195.36ms   //  202.05ms :  2,000,000 inserts creating 2M table
+    // 2234.84ms   // 2187.74ms   // 2240.38ms : 20,000,000 inserts creating 20M table
+    //   55.68ms   //   53.72ms   //   57.31ms : 500,000 inserts & deletes in 2K table
+    //   63.43ms   //   61.99ms   //   65.73ms : 500,000 inserts & deletes in 20K table
+    //   80.04ms   //   77.96ms   //   81.83ms : 500,000 inserts & deletes in 200K table
+    //  100.42ms   //   97.40ms   //  102.39ms : 500,000 inserts & deletes in 2M table
+    //  119.71ms   //  120.59ms   //  121.63ms : 500,000 inserts & deletes in 20M table
+    //  185.28ms   //  195.15ms   //  187.74ms : 500,000 inserts & deletes in 200M table
+    //   15.58ms   //   14.79ms   //   15.52ms : 200,000 inserts creating 200K table with varying key spacing
+
+    return (((size_t) hash << 16 << 16) | hash) ^ seed;
+  } else if (len == 8 && sizeof(size_t) == 8) {
+    size_t hash = d[0] | (d[1] << 8) | (d[2] << 16) | (d[3] << 24);
+    hash |= (size_t) (d[4] | (d[5] << 8) | (d[6] << 16) | (d[7] << 24)) << 16 << 16; // avoid warning if size_t == 4
+    hash ^= seed;
+    hash = (~hash) + (hash << 21);
+    hash ^= STBDS_ROTATE_RIGHT(hash,24);
+    hash *= 265;
+    hash ^= STBDS_ROTATE_RIGHT(hash,14);
+    hash ^= seed;
+    hash *= 21;
+    hash ^= STBDS_ROTATE_RIGHT(hash,28);
+    hash += (hash << 31);
+    hash = (~hash) + (hash << 18);
+    return hash;
+  } else {
+    return stbds_siphash_bytes(p,len,seed);
+  }
+#endif
+}
+#ifdef _MSC_VER
+#pragma warning(pop)
+#endif
+
+
+static int stbds_is_key_equal(void *a, size_t elemsize, void *key, size_t keysize, size_t keyoffset, int mode, size_t i)
+{
+  if (mode >= STBDS_HM_STRING)
+    return 0==strcmp((char *) key, * (char **) ((char *) a + elemsize*i + keyoffset));
+  else
+    return 0==memcmp(key, (char *) a + elemsize*i + keyoffset, keysize);
+}
+
+#define STBDS_HASH_TO_ARR(x,elemsize) ((char*) (x) - (elemsize))
+#define STBDS_ARR_TO_HASH(x,elemsize) ((char*) (x) + (elemsize))
+
+#define stbds_hash_table(a)  ((stbds_hash_index *) stbds_header(a)->hash_table)
+
+void stbds_hmfree_func(void *a, size_t elemsize)
+{
+  if (a == NULL) return;
+  if (stbds_hash_table(a) != NULL) {
+    if (stbds_hash_table(a)->string.mode == STBDS_SH_STRDUP) {
+      size_t i;
+      // skip 0th element, which is default
+      for (i=1; i < stbds_header(a)->length; ++i)
+        STBDS_FREE(NULL, *(char**) ((char *) a + elemsize*i));
+    }
+    stbds_strreset(&stbds_hash_table(a)->string);
+  }
+  STBDS_FREE(NULL, stbds_header(a)->hash_table);
+  STBDS_FREE(NULL, stbds_header(a));
+}
+
+static ptrdiff_t stbds_hm_find_slot(void *a, size_t elemsize, void *key, size_t keysize, size_t keyoffset, int mode)
+{
+  void *raw_a = STBDS_HASH_TO_ARR(a,elemsize);
+  stbds_hash_index *table = stbds_hash_table(raw_a);
+  size_t hash = mode >= STBDS_HM_STRING ? stbds_hash_string((char*)key,table->seed) : stbds_hash_bytes(key, keysize,table->seed);
+  size_t step = STBDS_BUCKET_LENGTH;
+  size_t limit,i;
+  size_t pos;
+  stbds_hash_bucket *bucket;
+
+  if (hash < 2) hash += 2; // stored hash values are forbidden from being 0, so we can detect empty slots
+
+  pos = stbds_probe_position(hash, table->slot_count, table->slot_count_log2);
+
+  for (;;) {
+    STBDS_STATS(++stbds_hash_probes);
+    bucket = &table->storage[pos >> STBDS_BUCKET_SHIFT];
+
+    // start searching from pos to end of bucket, this should help performance on small hash tables that fit in cache
+    for (i=pos & STBDS_BUCKET_MASK; i < STBDS_BUCKET_LENGTH; ++i) {
+      if (bucket->hash[i] == hash) {
+        if (stbds_is_key_equal(a, elemsize, key, keysize, keyoffset, mode, bucket->index[i])) {
+          return (pos & ~STBDS_BUCKET_MASK)+i;
+        }
+      } else if (bucket->hash[i] == STBDS_HASH_EMPTY) {
+        return -1;
+      }
+    }
+
+    // search from beginning of bucket to pos
+    limit = pos & STBDS_BUCKET_MASK;
+    for (i = 0; i < limit; ++i) {
+      if (bucket->hash[i] == hash) {
+        if (stbds_is_key_equal(a, elemsize, key, keysize, keyoffset, mode, bucket->index[i])) {
+          return (pos & ~STBDS_BUCKET_MASK)+i;
+        }
+      } else if (bucket->hash[i] == STBDS_HASH_EMPTY) {
+        return -1;
+      }
+    }
+
+    // quadratic probing
+    pos += step;
+    step += STBDS_BUCKET_LENGTH;
+    pos &= (table->slot_count-1);
+  }
+  /* NOTREACHED */
+}
+
+void * stbds_hmget_key_ts(void *a, size_t elemsize, void *key, size_t keysize, ptrdiff_t *temp, int mode)
+{
+  size_t keyoffset = 0;
+  if (a == NULL) {
+    // make it non-empty so we can return a temp
+    a = stbds_arrgrowf(0, elemsize, 0, 1);
+    stbds_header(a)->length += 1;
+    memset(a, 0, elemsize);
+    *temp = STBDS_INDEX_EMPTY;
+    // adjust a to point after the default element
+    return STBDS_ARR_TO_HASH(a,elemsize);
+  } else {
+    stbds_hash_index *table;
+    void *raw_a = STBDS_HASH_TO_ARR(a,elemsize);
+    // adjust a to point to the default element
+    table = (stbds_hash_index *) stbds_header(raw_a)->hash_table;
+    if (table == 0) {
+      *temp = -1;
+    } else {
+      ptrdiff_t slot = stbds_hm_find_slot(a, elemsize, key, keysize, keyoffset, mode);
+      if (slot < 0) {
+        *temp = STBDS_INDEX_EMPTY;
+      } else {
+        stbds_hash_bucket *b = &table->storage[slot >> STBDS_BUCKET_SHIFT];
+        *temp = b->index[slot & STBDS_BUCKET_MASK];
+      }
+    }
+    return a;
+  }
+}
+
+void * stbds_hmget_key(void *a, size_t elemsize, void *key, size_t keysize, int mode)
+{
+  ptrdiff_t temp;
+  void *p = stbds_hmget_key_ts(a, elemsize, key, keysize, &temp, mode);
+  stbds_temp(STBDS_HASH_TO_ARR(p,elemsize)) = temp;
+  return p;
+}
+
+void * stbds_hmput_default(void *a, size_t elemsize)
+{
+  // three cases:
+  //   a is NULL <- allocate
+  //   a has a hash table but no entries, because of shmode <- grow
+  //   a has entries <- do nothing
+  if (a == NULL || stbds_header(STBDS_HASH_TO_ARR(a,elemsize))->length == 0) {
+    a = stbds_arrgrowf(a ? STBDS_HASH_TO_ARR(a,elemsize) : NULL, elemsize, 0, 1);
+    stbds_header(a)->length += 1;
+    memset(a, 0, elemsize);
+    a=STBDS_ARR_TO_HASH(a,elemsize);
+  }
+  return a;
+}
+
+static char *stbds_strdup(char *str);
+
+void *stbds_hmput_key(void *a, size_t elemsize, void *key, size_t keysize, int mode)
+{
+  size_t keyoffset=0;
+  void *raw_a;
+  stbds_hash_index *table;
+
+  if (a == NULL) {
+    a = stbds_arrgrowf(0, elemsize, 0, 1);
+    memset(a, 0, elemsize);
+    stbds_header(a)->length += 1;
+    // adjust a to point AFTER the default element
+    a = STBDS_ARR_TO_HASH(a,elemsize);
+  }
+
+  // adjust a to point to the default element
+  raw_a = a;
+  a = STBDS_HASH_TO_ARR(a,elemsize);
+
+  table = (stbds_hash_index *) stbds_header(a)->hash_table;
+
+  if (table == NULL || table->used_count >= table->used_count_threshold) {
+    stbds_hash_index *nt;
+    size_t slot_count;
+
+    slot_count = (table == NULL) ? STBDS_BUCKET_LENGTH : table->slot_count*2;
+    nt = stbds_make_hash_index(slot_count, table);
+    if (table)
+      STBDS_FREE(NULL, table);
+    else
+      nt->string.mode = mode >= STBDS_HM_STRING ? STBDS_SH_DEFAULT : 0;
+    stbds_header(a)->hash_table = table = nt;
+    STBDS_STATS(++stbds_hash_grow);
+  }
+
+  // we iterate hash table explicitly because we want to track if we saw a tombstone
+  {
+    size_t hash = mode >= STBDS_HM_STRING ? stbds_hash_string((char*)key,table->seed) : stbds_hash_bytes(key, keysize,table->seed);
+    size_t step = STBDS_BUCKET_LENGTH;
+    size_t pos;
+    ptrdiff_t tombstone = -1;
+    stbds_hash_bucket *bucket;
+
+    // stored hash values are forbidden from being 0, so we can detect empty slots to early out quickly
+    if (hash < 2) hash += 2;
+
+    pos = stbds_probe_position(hash, table->slot_count, table->slot_count_log2);
+
+    for (;;) {
+      size_t limit, i;
+      STBDS_STATS(++stbds_hash_probes);
+      bucket = &table->storage[pos >> STBDS_BUCKET_SHIFT];
+
+      // start searching from pos to end of bucket
+      for (i=pos & STBDS_BUCKET_MASK; i < STBDS_BUCKET_LENGTH; ++i) {
+        if (bucket->hash[i] == hash) {
+          if (stbds_is_key_equal(raw_a, elemsize, key, keysize, keyoffset, mode, bucket->index[i])) {
+            stbds_temp(a) = bucket->index[i];
+            return STBDS_ARR_TO_HASH(a,elemsize);
+          }
+        } else if (bucket->hash[i] == 0) {
+          pos = (pos & ~STBDS_BUCKET_MASK) + i;
+          goto found_empty_slot;
+        } else if (tombstone < 0) {
+          if (bucket->index[i] == STBDS_INDEX_DELETED)
+            tombstone = (ptrdiff_t) ((pos & ~STBDS_BUCKET_MASK) + i);
+        }
+      }
+
+      // search from beginning of bucket to pos
+      limit = pos & STBDS_BUCKET_MASK;
+      for (i = 0; i < limit; ++i) {
+        if (bucket->hash[i] == hash) {
+          if (stbds_is_key_equal(raw_a, elemsize, key, keysize, keyoffset, mode, bucket->index[i])) {
+            stbds_temp(a) = bucket->index[i];
+            return STBDS_ARR_TO_HASH(a,elemsize);
+          }
+        } else if (bucket->hash[i] == 0) {
+          pos = (pos & ~STBDS_BUCKET_MASK) + i;
+          goto found_empty_slot;
+        } else if (tombstone < 0) {
+          if (bucket->index[i] == STBDS_INDEX_DELETED)
+            tombstone = (ptrdiff_t) ((pos & ~STBDS_BUCKET_MASK) + i);
+        }
+      }
+
+      // quadratic probing
+      pos += step;
+      step += STBDS_BUCKET_LENGTH;
+      pos &= (table->slot_count-1);
+    }
+   found_empty_slot:
+    if (tombstone >= 0) {
+      pos = tombstone;
+      --table->tombstone_count;
+    }
+    ++table->used_count;
+
+    {
+      ptrdiff_t j = (ptrdiff_t) stbds_arrlen(a);
+      // we want to do stbds_arraddn(1), but we can't use the macros since we don't have something of the right type
+      if ((size_t) j+1 > stbds_arrcap(a))
+        *(void **) &a = stbds_arrgrowf(a, elemsize, 1, 0);
+      raw_a = STBDS_ARR_TO_HASH(a,elemsize);
+
+      STBDS_ASSERT((size_t) j+1 <= stbds_arrcap(a));
+      stbds_header(a)->length = j+1;
+      bucket = &table->storage[pos >> STBDS_BUCKET_SHIFT];
+      bucket->hash[pos & STBDS_BUCKET_MASK] = hash;
+      bucket->index[pos & STBDS_BUCKET_MASK] = j-1;
+      stbds_temp(a) = j-1;
+
+      switch (table->string.mode) {
+         case STBDS_SH_STRDUP:  stbds_temp_key(a) = *(char **) ((char *) a + elemsize*j) = stbds_strdup((char*) key); break;
+         case STBDS_SH_ARENA:   stbds_temp_key(a) = *(char **) ((char *) a + elemsize*j) = stbds_stralloc(&table->string, (char*)key); break;
+         case STBDS_SH_DEFAULT: stbds_temp_key(a) = *(char **) ((char *) a + elemsize*j) = (char *) key; break;
+         default:               memcpy((char *) a + elemsize*j, key, keysize); break;
+      }
+    }
+    return STBDS_ARR_TO_HASH(a,elemsize);
+  }
+}
+
+void * stbds_shmode_func(size_t elemsize, int mode)
+{
+  void *a = stbds_arrgrowf(0, elemsize, 0, 1);
+  stbds_hash_index *h;
+  memset(a, 0, elemsize);
+  stbds_header(a)->length = 1;
+  stbds_header(a)->hash_table = h = (stbds_hash_index *) stbds_make_hash_index(STBDS_BUCKET_LENGTH, NULL);
+  h->string.mode = (unsigned char) mode;
+  return STBDS_ARR_TO_HASH(a,elemsize);
+}
+
+void * stbds_hmdel_key(void *a, size_t elemsize, void *key, size_t keysize, size_t keyoffset, int mode)
+{
+  if (a == NULL) {
+    return 0;
+  } else {
+    stbds_hash_index *table;
+    void *raw_a = STBDS_HASH_TO_ARR(a,elemsize);
+    table = (stbds_hash_index *) stbds_header(raw_a)->hash_table;
+    stbds_temp(raw_a) = 0;
+    if (table == 0) {
+      return a;
+    } else {
+      ptrdiff_t slot;
+      slot = stbds_hm_find_slot(a, elemsize, key, keysize, keyoffset, mode);
+      if (slot < 0)
+        return a;
+      else {
+        stbds_hash_bucket *b = &table->storage[slot >> STBDS_BUCKET_SHIFT];
+        int i = slot & STBDS_BUCKET_MASK;
+        ptrdiff_t old_index = b->index[i];
+        ptrdiff_t final_index = (ptrdiff_t) stbds_arrlen(raw_a)-1-1; // minus one for the raw_a vs a, and minus one for 'last'
+        STBDS_ASSERT(slot < (ptrdiff_t) table->slot_count);
+        --table->used_count;
+        ++table->tombstone_count;
+        stbds_temp(raw_a) = 1;
+        STBDS_ASSERT(table->used_count >= 0);
+        //STBDS_ASSERT(table->tombstone_count < table->slot_count/4);
+        b->hash[i] = STBDS_HASH_DELETED;
+        b->index[i] = STBDS_INDEX_DELETED;
+
+        if (mode == STBDS_HM_STRING && table->string.mode == STBDS_SH_STRDUP)
+          STBDS_FREE(NULL, *(char**) ((char *) a+elemsize*old_index));
+
+        // if indices are the same, memcpy is a no-op, but back-pointer-fixup will fail, so skip
+        if (old_index != final_index) {
+          // swap delete
+          memmove((char*) a + elemsize*old_index, (char*) a + elemsize*final_index, elemsize);
+
+          // now find the slot for the last element
+          if (mode == STBDS_HM_STRING)
+            slot = stbds_hm_find_slot(a, elemsize, *(char**) ((char *) a+elemsize*old_index + keyoffset), keysize, keyoffset, mode);
+          else
+            slot = stbds_hm_find_slot(a, elemsize,  (char* ) a+elemsize*old_index + keyoffset, keysize, keyoffset, mode);
+          STBDS_ASSERT(slot >= 0);
+          b = &table->storage[slot >> STBDS_BUCKET_SHIFT];
+          i = slot & STBDS_BUCKET_MASK;
+          STBDS_ASSERT(b->index[i] == final_index);
+          b->index[i] = old_index;
+        }
+        stbds_header(raw_a)->length -= 1;
+
+        if (table->used_count < table->used_count_shrink_threshold && table->slot_count > STBDS_BUCKET_LENGTH) {
+          stbds_header(raw_a)->hash_table = stbds_make_hash_index(table->slot_count>>1, table);
+          STBDS_FREE(NULL, table);
+          STBDS_STATS(++stbds_hash_shrink);
+        } else if (table->tombstone_count > table->tombstone_count_threshold) {
+          stbds_header(raw_a)->hash_table = stbds_make_hash_index(table->slot_count   , table);
+          STBDS_FREE(NULL, table);
+          STBDS_STATS(++stbds_hash_rebuild);
+        }
+
+        return a;
+      }
+    }
+  }
+  /* NOTREACHED */
+}
+
+static char *stbds_strdup(char *str)
+{
+  // to keep replaceable allocator simple, we don't want to use strdup.
+  // rolling our own also avoids problem of strdup vs _strdup
+  size_t len = strlen(str)+1;
+  char *p = (char*) STBDS_REALLOC(NULL, 0, len);
+  memmove(p, str, len);
+  return p;
+}
+
+#ifndef STBDS_STRING_ARENA_BLOCKSIZE_MIN
+#define STBDS_STRING_ARENA_BLOCKSIZE_MIN  512u
+#endif
+#ifndef STBDS_STRING_ARENA_BLOCKSIZE_MAX
+#define STBDS_STRING_ARENA_BLOCKSIZE_MAX  (1u<<20)
+#endif
+
+char *stbds_stralloc(stbds_string_arena *a, char *str)
+{
+  char *p;
+  size_t len = strlen(str)+1;
+  if (len > a->remaining) {
+    // compute the next blocksize
+    size_t blocksize = a->block;
+
+    // size is 512, 512, 1024, 1024, 2048, 2048, 4096, 4096, etc., so that
+    // there are log(SIZE) allocations to free when we destroy the table
+    blocksize = (size_t) (STBDS_STRING_ARENA_BLOCKSIZE_MIN) << (blocksize>>1);
+
+    // if size is under 1M, advance to next blocktype
+    if (blocksize < (size_t)(STBDS_STRING_ARENA_BLOCKSIZE_MAX))
+      ++a->block;
+
+    if (len > blocksize) {
+      // if string is larger than blocksize, then just allocate the full size.
+      // note that we still advance string_block so block size will continue
+      // increasing, so e.g. if somebody only calls this with 1000-long strings,
+      // eventually the arena will start doubling and handling those as well
+      stbds_string_block *sb = (stbds_string_block *) STBDS_REALLOC(NULL, 0, sizeof(*sb)-8 + len);
+      memmove(sb->storage, str, len);
+      if (a->storage) {
+        // insert it after the first element, so that we don't waste the space there
+        sb->next = a->storage->next;
+        a->storage->next = sb;
+      } else {
+        sb->next = 0;
+        a->storage = sb;
+        a->remaining = 0; // this is redundant, but good for clarity
+      }
+      return sb->storage;
+    } else {
+      stbds_string_block *sb = (stbds_string_block *) STBDS_REALLOC(NULL, 0, sizeof(*sb)-8 + blocksize);
+      sb->next = a->storage;
+      a->storage = sb;
+      a->remaining = blocksize;
+    }
+  }
+
+  STBDS_ASSERT(len <= a->remaining);
+  p = a->storage->storage + a->remaining - len;
+  a->remaining -= len;
+  memmove(p, str, len);
+  return p;
+}
+
+void stbds_strreset(stbds_string_arena *a)
+{
+  stbds_string_block *x,*y;
+  x = a->storage;
+  while (x) {
+    y = x->next;
+    STBDS_FREE(NULL, x);
+    x = y;
+  }
+  memset(a, 0, sizeof(*a));
+}
+
+#endif
+
+//////////////////////////////////////////////////////////////////////////////
+//
+//   UNIT TESTS
+//
+
+#ifdef STBDS_UNIT_TESTS
+#include <stdio.h>
+#ifdef STBDS_ASSERT_WAS_UNDEFINED
+#undef STBDS_ASSERT
+#endif
+#ifndef STBDS_ASSERT
+#define STBDS_ASSERT assert
+#include <assert.h>
+#endif
+
+typedef struct { int key,b,c,d; } stbds_struct;
+typedef struct { int key[2],b,c,d; } stbds_struct2;
+
+static char buffer[256];
+char *strkey(int n)
+{
+#if defined(_WIN32) && defined(__STDC_WANT_SECURE_LIB__)
+   sprintf_s(buffer, sizeof(buffer), "test_%d", n);
+#else
+   sprintf(buffer, "test_%d", n);
+#endif
+   return buffer;
+}
+
+void stbds_unit_tests(void)
+{
+#if defined(_MSC_VER) && _MSC_VER <= 1200 && defined(__cplusplus)
+  // VC6 C++ doesn't like the template<> trick on unnamed structures, so do nothing!
+  STBDS_ASSERT(0);
+#else
+  const int testsize = 100000;
+  const int testsize2 = testsize/20;
+  int *arr=NULL;
+  struct { int   key;        int value; }  *intmap  = NULL;
+  struct { char *key;        int value; }  *strmap  = NULL, s;
+  struct { stbds_struct key; int value; }  *map     = NULL;
+  stbds_struct                             *map2    = NULL;
+  stbds_struct2                            *map3    = NULL;
+  stbds_string_arena                        sa      = { 0 };
+  int key3[2] = { 1,2 };
+  ptrdiff_t temp;
+
+  int i,j;
+
+  STBDS_ASSERT(arrlen(arr)==0);
+  for (i=0; i < 20000; i += 50) {
+    for (j=0; j < i; ++j)
+      arrpush(arr,j);
+    arrfree(arr);
+  }
+
+  for (i=0; i < 4; ++i) {
+    arrpush(arr,1); arrpush(arr,2); arrpush(arr,3); arrpush(arr,4);
+    arrdel(arr,i);
+    arrfree(arr);
+    arrpush(arr,1); arrpush(arr,2); arrpush(arr,3); arrpush(arr,4);
+    arrdelswap(arr,i);
+    arrfree(arr);
+  }
+
+  for (i=0; i < 5; ++i) {
+    arrpush(arr,1); arrpush(arr,2); arrpush(arr,3); arrpush(arr,4);
+    stbds_arrins(arr,i,5);
+    STBDS_ASSERT(arr[i] == 5);
+    if (i < 4)
+      STBDS_ASSERT(arr[4] == 4);
+    arrfree(arr);
+  }
+
+  i = 1;
+  STBDS_ASSERT(hmgeti(intmap,i) == -1);
+  hmdefault(intmap, -2);
+  STBDS_ASSERT(hmgeti(intmap, i) == -1);
+  STBDS_ASSERT(hmget (intmap, i) == -2);
+  for (i=0; i < testsize; i+=2)
+    hmput(intmap, i, i*5);
+  for (i=0; i < testsize; i+=1) {
+    if (i & 1) STBDS_ASSERT(hmget(intmap, i) == -2 );
+    else       STBDS_ASSERT(hmget(intmap, i) == i*5);
+    if (i & 1) STBDS_ASSERT(hmget_ts(intmap, i, temp) == -2 );
+    else       STBDS_ASSERT(hmget_ts(intmap, i, temp) == i*5);
+  }
+  for (i=0; i < testsize; i+=2)
+    hmput(intmap, i, i*3);
+  for (i=0; i < testsize; i+=1)
+    if (i & 1) STBDS_ASSERT(hmget(intmap, i) == -2 );
+    else       STBDS_ASSERT(hmget(intmap, i) == i*3);
+  for (i=2; i < testsize; i+=4)
+    hmdel(intmap, i); // delete half the entries
+  for (i=0; i < testsize; i+=1)
+    if (i & 3) STBDS_ASSERT(hmget(intmap, i) == -2 );
+    else       STBDS_ASSERT(hmget(intmap, i) == i*3);
+  for (i=0; i < testsize; i+=1)
+    hmdel(intmap, i); // delete the rest of the entries
+  for (i=0; i < testsize; i+=1)
+    STBDS_ASSERT(hmget(intmap, i) == -2 );
+  hmfree(intmap);
+  for (i=0; i < testsize; i+=2)
+    hmput(intmap, i, i*3);
+  hmfree(intmap);
+
+  #if defined(__clang__) || defined(__GNUC__)
+  #ifndef __cplusplus
+  intmap = NULL;
+  hmput(intmap, 15, 7);
+  hmput(intmap, 11, 3);
+  hmput(intmap,  9, 5);
+  STBDS_ASSERT(hmget(intmap, 9) == 5);
+  STBDS_ASSERT(hmget(intmap, 11) == 3);
+  STBDS_ASSERT(hmget(intmap, 15) == 7);
+  #endif
+  #endif
+
+  for (i=0; i < testsize; ++i)
+    stralloc(&sa, strkey(i));
+  strreset(&sa);
+
+  {
+    s.key = "a", s.value = 1;
+    shputs(strmap, s);
+    STBDS_ASSERT(*strmap[0].key == 'a');
+    STBDS_ASSERT(strmap[0].key == s.key);
+    STBDS_ASSERT(strmap[0].value == s.value);
+    shfree(strmap);
+  }
+
+  {
+    s.key = "a", s.value = 1;
+    sh_new_strdup(strmap);
+    shputs(strmap, s);
+    STBDS_ASSERT(*strmap[0].key == 'a');
+    STBDS_ASSERT(strmap[0].key != s.key);
+    STBDS_ASSERT(strmap[0].value == s.value);
+    shfree(strmap);
+  }
+
+  {
+    s.key = "a", s.value = 1;
+    sh_new_arena(strmap);
+    shputs(strmap, s);
+    STBDS_ASSERT(*strmap[0].key == 'a');
+    STBDS_ASSERT(strmap[0].key != s.key);
+    STBDS_ASSERT(strmap[0].value == s.value);
+    shfree(strmap);
+  }
+
+  for (j=0; j < 2; ++j) {
+    STBDS_ASSERT(shgeti(strmap,"foo") == -1);
+    if (j == 0)
+      sh_new_strdup(strmap);
+    else
+      sh_new_arena(strmap);
+    STBDS_ASSERT(shgeti(strmap,"foo") == -1);
+    shdefault(strmap, -2);
+    STBDS_ASSERT(shgeti(strmap,"foo") == -1);
+    for (i=0; i < testsize; i+=2)
+      shput(strmap, strkey(i), i*3);
+    for (i=0; i < testsize; i+=1)
+      if (i & 1) STBDS_ASSERT(shget(strmap, strkey(i)) == -2 );
+      else       STBDS_ASSERT(shget(strmap, strkey(i)) == i*3);
+    for (i=2; i < testsize; i+=4)
+      shdel(strmap, strkey(i)); // delete half the entries
+    for (i=0; i < testsize; i+=1)
+      if (i & 3) STBDS_ASSERT(shget(strmap, strkey(i)) == -2 );
+      else       STBDS_ASSERT(shget(strmap, strkey(i)) == i*3);
+    for (i=0; i < testsize; i+=1)
+      shdel(strmap, strkey(i)); // delete the rest of the entries
+    for (i=0; i < testsize; i+=1)
+      STBDS_ASSERT(shget(strmap, strkey(i)) == -2 );
+    shfree(strmap);
+  }
+
+  {
+    struct { char *key; char value; } *hash = NULL;
+    char name[4] = "jen";
+    shput(hash, "bob"   , 'h');
+    shput(hash, "sally" , 'e');
+    shput(hash, "fred"  , 'l');
+    shput(hash, "jen"   , 'x');
+    shput(hash, "doug"  , 'o');
+
+    shput(hash, name    , 'l');
+    shfree(hash);
+  }
+
+  for (i=0; i < testsize; i += 2) {
+    stbds_struct s = { i,i*2,i*3,i*4 };
+    hmput(map, s, i*5);
+  }
+
+  for (i=0; i < testsize; i += 1) {
+    stbds_struct s = { i,i*2,i*3  ,i*4 };
+    stbds_struct t = { i,i*2,i*3+1,i*4 };
+    if (i & 1) STBDS_ASSERT(hmget(map, s) == 0);
+    else       STBDS_ASSERT(hmget(map, s) == i*5);
+    if (i & 1) STBDS_ASSERT(hmget_ts(map, s, temp) == 0);
+    else       STBDS_ASSERT(hmget_ts(map, s, temp) == i*5);
+    //STBDS_ASSERT(hmget(map, t.key) == 0);
+  }
+
+  for (i=0; i < testsize; i += 2) {
+    stbds_struct s = { i,i*2,i*3,i*4 };
+    hmputs(map2, s);
+  }
+  hmfree(map);
+
+  for (i=0; i < testsize; i += 1) {
+    stbds_struct s = { i,i*2,i*3,i*4 };
+    stbds_struct t = { i,i*2,i*3+1,i*4 };
+    if (i & 1) STBDS_ASSERT(hmgets(map2, s.key).d == 0);
+    else       STBDS_ASSERT(hmgets(map2, s.key).d == i*4);
+    //STBDS_ASSERT(hmgetp(map2, t.key) == 0);
+  }
+  hmfree(map2);
+
+  for (i=0; i < testsize; i += 2) {
+    stbds_struct2 s = { { i,i*2 }, i*3,i*4, i*5 };
+    hmputs(map3, s);
+  }
+  for (i=0; i < testsize; i += 1) {
+    stbds_struct2 s = { { i,i*2}, i*3, i*4, i*5 };
+    stbds_struct2 t = { { i,i*2}, i*3+1, i*4, i*5 };
+    if (i & 1) STBDS_ASSERT(hmgets(map3, s.key).d == 0);
+    else       STBDS_ASSERT(hmgets(map3, s.key).d == i*5);
+    //STBDS_ASSERT(hmgetp(map3, t.key) == 0);
+  }
+#endif
+}
+#endif
+
+
+/*
+------------------------------------------------------------------------------
+This software is available under 2 licenses -- choose whichever you prefer.
+------------------------------------------------------------------------------
+ALTERNATIVE A - MIT License
+Copyright (c) 2019 Sean Barrett
+Permission is hereby granted, free of charge, to any person obtaining a copy of
+this software and associated documentation files (the "Software"), to deal in
+the Software without restriction, including without limitation the rights to
+use, copy, modify, merge, publish, distribute, sublicense, and/or sell copies
+of the Software, and to permit persons to whom the Software is furnished to do
+so, subject to the following conditions:
+The above copyright notice and this permission notice shall be included in all
+copies or substantial portions of the Software.
+THE SOFTWARE IS PROVIDED "AS IS", WITHOUT WARRANTY OF ANY KIND, EXPRESS OR
+IMPLIED, INCLUDING BUT NOT LIMITED TO THE WARRANTIES OF MERCHANTABILITY,
+FITNESS FOR A PARTICULAR PURPOSE AND NONINFRINGEMENT. IN NO EVENT SHALL THE
+AUTHORS OR COPYRIGHT HOLDERS BE LIABLE FOR ANY CLAIM, DAMAGES OR OTHER
+LIABILITY, WHETHER IN AN ACTION OF CONTRACT, TORT OR OTHERWISE, ARISING FROM,
+OUT OF OR IN CONNECTION WITH THE SOFTWARE OR THE USE OR OTHER DEALINGS IN THE
+SOFTWARE.
+------------------------------------------------------------------------------
+ALTERNATIVE B - Public Domain (www.unlicense.org)
+This is free and unencumbered software released into the public domain.
+Anyone is free to copy, modify, publish, use, compile, sell, or distribute this
+software, either in source code form or as a compiled binary, for any purpose,
+commercial or non-commercial, and by any means.
+In jurisdictions that recognize copyright laws, the author or authors of this
+software dedicate any and all copyright interest in the software to the public
+domain. We make this dedication for the benefit of the public at large and to
+the detriment of our heirs and successors. We intend this dedication to be an
+overt act of relinquishment in perpetuity of all present and future rights to
+this software under copyright law.
+THE SOFTWARE IS PROVIDED "AS IS", WITHOUT WARRANTY OF ANY KIND, EXPRESS OR
+IMPLIED, INCLUDING BUT NOT LIMITED TO THE WARRANTIES OF MERCHANTABILITY,
+FITNESS FOR A PARTICULAR PURPOSE AND NONINFRINGEMENT. IN NO EVENT SHALL THE
+AUTHORS BE LIABLE FOR ANY CLAIM, DAMAGES OR OTHER LIABILITY, WHETHER IN AN
+ACTION OF CONTRACT, TORT OR OTHERWISE, ARISING FROM, OUT OF OR IN CONNECTION
+WITH THE SOFTWARE OR THE USE OR OTHER DEALINGS IN THE SOFTWARE.
+------------------------------------------------------------------------------
+*/
diff --git a/Engine/Source/Developer/ShaderPreprocessor/ShaderPreprocessor.Build.cs b/Engine/Source/Developer/ShaderPreprocessor/ShaderPreprocessor.Build.cs
--- a/Engine/Source/Developer/ShaderPreprocessor/ShaderPreprocessor.Build.cs
+++ b/Engine/Source/Developer/ShaderPreprocessor/ShaderPreprocessor.Build.cs
@@ -14,5 +14,8 @@ public class ShaderPreprocessor : ModuleRules
 			);
 
 		AddEngineThirdPartyPrivateStaticDependencies(Target, "MCPP");
+
+		PrivateDefinitions.Add("STB_DS_IMPLEMENTATION");
+		PrivateDefinitions.Add("STB_ALLOC_IMPLEMENTATION");
 	}
 }
