[shaders] additions & fixes to STB preprocessor for shader preprocessing. With the following all Lyra shaders compile and generate identical bytecode on PC SM5/6.
- add explicit extern "C" on preprocessor methods to allow linking from C++ code
- add custom callbacks for include path resolution and file loading/unloading
- add padding space following macro expansion ending with a > to work around DXC bug (not handling >> correctly if it occurs at the end of a templated type)
- adding explicit handling for HLSL infinity constant (1.#INF); the # character outside of a directive was previously putting the preprocessor into an error state
- change line directive output to use #line format rather than # (DXC is ok with the latter but minifier doesn't currently handle it), and also always include the filename (so when minifying don't have to be careful not to strip the line directive with the filename in it); this also matches mcpp behaviour
- adding const-correctness and casts where needed to avoid const casts/compile errors in calling C++ code and satisfy Mac/Intellisense compilers
- bugfix: parsing of '.' was immediately putting the scanning state machine into "number" mode resulting in skipped macro expansion in code of the form "var.MACRO.var2"/numbers; added extra state transitions and changing the default transition when encountering '.' to an explicit mode that can allow either transitioning to parsing identifiers or numbers depending on following characters.
- bugfix: macros with an empty parameter list were incorrectly terminating said list on expansion (adding extra rparens)
- bugfix: set file load mode as "callback" when callback used for loading, and check the mode in the free callbacks
- various minor changes to satisfy Clang code analysis

diff --git a/Engine/Source/Developer/ShaderPreprocessor/Private/stb_preprocess/preprocessor.c b/Engine/Source/Developer/ShaderPreprocessor/Private/stb_preprocess/preprocessor.c
--- a/Engine/Source/Developer/ShaderPreprocessor/Private/stb_preprocess/preprocessor.c
+++ b/Engine/Source/Developer/ShaderPreprocessor/Private/stb_preprocess/preprocessor.c
@@ -151,6 +151,10 @@ typedef struct
 	void* value;
 } shpair;
 
+static loadfile_callback_func loadfile_callback;
+static freefile_callback_func freefile_callback;
+static resolveinclude_callback_func resolveinclude_callback;
+
 #define HASH_EMPTY_MARKER 0
 #define HASH_TOMBSTONE 1
 
@@ -198,17 +202,18 @@ typedef struct pp_context
 	int counter;
 
 	struct stb_arena macro_arena;
+	void* custom_context;
 } pp_context;
 
 typedef struct parse_state
 {
 	pp_context* context;
 
-	char* filename;	 // this is cached in last_output_filename so must remain valid indefinitely
-	char* last_output_filename;
-	char* stringified_filename;	 // this is NULL until first time __FILE__ is used, then it's cached in macro_arena
+	const char* filename; // this is cached in last_output_filename so must remain valid indefinitely
+	const char* last_output_filename;
+	char* stringified_filename; // this is NULL until first time __FILE__ is used, then it's cached in macro_arena
 
-	char* src;
+	const char* src;
 	size_t src_offset;
 	size_t src_length;
 
@@ -235,11 +240,11 @@ typedef struct parse_state
 //  UTILITIES
 //
 
-static char* arena_alloc_padded_string(struct stb_arena* a, char* text, size_t stringlen)
+static char* arena_alloc_padded_string(struct stb_arena* a, const char* text, size_t stringlen)
 {
 	// allocate a string with whitespace at each end to prevent accidental token-pasting
 	size_t final_length = stringlen + 3;
-	char* mem = stb_arena_alloc_aligned(a, final_length, 4);
+	char* mem = (char*)stb_arena_alloc_aligned(a, final_length, 4);
 	mem[0] = ' ';
 	memcpy(mem + 1, text, stringlen);
 	mem[stringlen + 1] = ' ';
@@ -247,6 +252,18 @@ static char* arena_alloc_padded_string(struct stb_arena* a, char* text, size_t s
 	return mem;
 }
 
+static char* arena_alloc_string_trailingspace(struct stb_arena* a, const char* text, size_t stringlen)
+{
+	// allocate a string with a trailing space appended
+	size_t final_length = stringlen + 2;
+	char* mem = (char*)stb_arena_alloc_aligned(a, final_length, 4);
+	memcpy(mem, text, stringlen);
+	mem[stringlen] = ' ';
+	mem[stringlen + 1] = 0;
+	return mem;
+
+}
+
 #define MIN3(a, b, c) ((a) < (b) ? ((a) < (c) ? (a) : (c)) : ((b) < (c) ? (b) : (c)))
 
 static int compute_edit_distance(const char* s1, const char* s2, int max_dist)
@@ -278,6 +295,7 @@ static int compute_edit_distance(const char* s1, const char* s2, int max_dist)
 		for (y = 1, lastdiag = x - 1; y <= s1len; y++)
 		{
 			olddiag = column[y];
+			assert(y - 1 < strlen(s1) && y > 1);
 			column[y] = MIN3(column[y] + 1, column[y - 1] + 1, lastdiag + (s1[y - 1] == s2[x - 1] ? 0 : 1));
 			lastdiag = olddiag;
 		}
@@ -383,11 +401,8 @@ static void output_line_directive(parse_state* cs)
 	char* q;
 	char* out = cs->dest;
 
-	// output the filename if it's changed
-	char* filename = (cs->last_output_filename && cs->last_output_filename == cs->filename ? "" : cs->filename);
-
 	// compute an upper bound on the amount of data to be written
-	size_t filename_len = strlen(filename);
+	size_t filename_len = strlen(cs->filename);
 	size_t expansion_length = 6 + 16 + filename_len + 2 + 2 + 1;
 
 	if (cs->fast_dest)
@@ -402,18 +417,24 @@ static void output_line_directive(parse_state* cs)
 
 	// write the line directive
 	stb_assume(q != NULL);
+	
+	// ensure line directives start on a new line
+	if (q[-1] != '\n')
+		*q++ = '\n';
+
 	*q++ = '#';
+	*q++ = 'l';
+	*q++ = 'i';
+	*q++ = 'n';
+	*q++ = 'e';
 	*q++ = ' ';
 	q = preprocessor_line_dtoa(q, cs->src_line_number);
-	if (filename[0])
-	{
-		*q++ = ' ';
-		*q++ = '"';
-		strcpy(q, filename);  // @TODO cache stringified filename
-		q += filename_len;
-		*q++ = '"';
-		cs->last_output_filename = filename;
-	}
+	*q++ = ' ';
+	*q++ = '"';
+	strcpy(q, cs->filename);  // @TODO cache stringified filename
+	q += filename_len;
+	*q++ = '"';
+	cs->last_output_filename = cs->filename;
 	*q++ = '\n';
 
 	cs->dest_line_number = cs->src_line_number;
@@ -449,11 +470,20 @@ void pp_set_did_you_mean_threshold(int mode)
 	did_you_mean_threshold = mode;
 }
 
+void fill_where(parse_state* stack, pp_diagnostic* diagnostic, int* line_number)
+{
+	pp_where where = { 0 };
+	where.filename = strdup(stack->filename);
+	where.line_number = (*line_number > 0 ? *line_number : stack->src_line_number);
+	where.column = 0;
+	*line_number = 0;
+	arrpush(diagnostic->where, where);
+}
+
 static int error_explicit_v(parse_state* ps, int code, int line_number, char* text, va_list va)
 {
 	stb_assume(ps != NULL);
 	parse_state* stack;
-	pp_where where = { 0 };
 	pp_diagnostic d;
 	size_t len;
 
@@ -473,9 +503,18 @@ static int error_explicit_v(parse_state* ps, int code, int line_number, char* te
 #endif
 
 	size_t size = len + 1;
-	d.message = malloc(size);
+	// sanity check added to satisfy clang static analyzer
+	if (size < 1)
+		return 1;
+
+	d.message = (char*)malloc(size);
+
+	// sanity check added to satisfy clang static analyzer
+	if (d.message == NULL)
+		return 1;
 
 	len = _vsnprintf(d.message, len, text, va);
+	
 	d.message[size - 1] = 0;
 
 	d.diagnostic_code = code;
@@ -484,17 +523,11 @@ static int error_explicit_v(parse_state* ps, int code, int line_number, char* te
 		d.error_level = PP_RESULT_MODE_warning;
 	d.where = 0;
 
-	stack = ps;
+	fill_where(ps, &d, &line_number);
+	stack = ps->parent;
 	while (stack != NULL)
 	{
-		if (stack->filename)
-		{
-			where.filename = strdup(stack->filename);
-			where.line_number = (line_number > 0 ? line_number : stack->src_line_number);
-			where.column = 0;
-			line_number = 0;
-			arrpush(d.where, where);
-		}
+		fill_where(stack, &d, &line_number);
 		stack = stack->parent;
 	}
 
@@ -562,13 +595,13 @@ static void error_supplement(parse_state* cs, int code, int line_number, char* t
 //
 //  based on XXHash32
 
-static unsigned int load_unaligned_32(void* ptr)
+static unsigned int load_unaligned_32(const void* ptr)
 {
 #ifdef NO_UNALIGNED_LOADS
 	unsigned char* data = ptr;
 	return data[0] | (data[1] << 8) | (data[2] << 16) | (data[3] << 24);
 #else
-	unsigned int* data = ptr;  // DOESN'T produce the same result for little-endian & big-endian
+	unsigned int* data = (unsigned int*)ptr;  // DOESN'T produce the same result for little-endian & big-endian
 	return *data;
 #endif
 }
@@ -579,7 +612,7 @@ static unsigned int load_unaligned_32(void* ptr)
 #define xxprime3 0xC2B2AE3Du
 
 // export hash function so compiler can use it as well
-unsigned int preprocessor_hash(char* data, size_t len)
+unsigned int preprocessor_hash(const char* data, size_t len)
 {
 	unsigned int hash;
 	if (len < 4)
@@ -660,7 +693,7 @@ void stringhash_init_table(pphash* ht, unsigned int size)
 	table_entry t = {0}, *table;
 	unsigned int i;
 	stb_assume(size > 0u);
-	table = malloc(sizeof(table[0]) * size);
+	table = (table_entry*)malloc(sizeof(table[0]) * size);
 	stb_assume(table != NULL);
 	ht->table_size = size;
 	ht->table_mask = size - 1;
@@ -670,7 +703,7 @@ void stringhash_init_table(pphash* ht, unsigned int size)
 		table[i] = t;  // hope for 8-byte-at-a-time init
 }
 
-void* stringhash_get(pphash* ht, char* key, size_t keylen)
+void* stringhash_get(pphash* ht, const char* key, size_t keylen)
 {
 	uint32 mask = ht->table_mask;
 	uint32 hash = preprocessor_hash(key, keylen);
@@ -750,7 +783,7 @@ int stringhash_delete(pphash* ht, char* key, size_t keylen)
 	{
 		i = table[slot].index;
 		if (ht->pair[i].key_length == keylen && 0 == memcmp(key, ht->pair[i].key, keylen))
-			goto delete;
+			goto del;
 	}
 	if (table[slot].hash == HASH_EMPTY_MARKER)
 		return 0;
@@ -764,12 +797,12 @@ int stringhash_delete(pphash* ht, char* key, size_t keylen)
 		{
 			i = table[slot].index;
 			if (ht->pair[i].key_length == keylen && 0 == memcmp(key, ht->pair[i].key, keylen))
-				goto delete;
+				goto del;
 		}
 		if (table[slot].hash == HASH_EMPTY_MARKER)
 			return 0;
 	}
-	delete : i = table[slot].index;
+	del : i = table[slot].index;
 	ht->pair[i].key = 0;
 	ht->pair[i].value = 0;
 	ht->pair[i].key_length = (size_t)ht->first_unused_pair;
@@ -802,7 +835,7 @@ static void stringhash_insert_ref(pphash* ht, table_entry e)
 	table[slot] = e;
 }
 
-void stringhash_put(pphash* ht, char* key, size_t keylen, void* value)
+void stringhash_put(pphash* ht, const char* key, size_t keylen, void* value)
 {
 	size_t n;
 	hashtable_pair p;
@@ -840,7 +873,7 @@ void stringhash_put(pphash* ht, char* key, size_t keylen, void* value)
 	if (table[slot].hash == HASH_TOMBSTONE)
 		--ht->count_tombstones;
 
-	p.key = stb_arena_alloc_aligned(&ht->arena, keylen + 1, 1);
+	p.key = (char*)stb_arena_alloc_aligned(&ht->arena, keylen + 1, 1);
 	p.key_length = keylen;
 	p.value = value;
 	memcpy(p.key, key, keylen);
@@ -1063,7 +1096,7 @@ static unsigned char pp_char_flags2[256];
 #define is_backslash_newline(p) ((p)[0] == '\\' && char_is_newline((p)[1]))
 
 // this is used when parsing a backslash-eliminated directive
-static char* preprocessor_skip_whitespace_simple(char* s)
+static const char* preprocessor_skip_whitespace_simple(const char* s)
 {
 	while (char_is_nonnewline_whitespace(*s))
 		++s;
@@ -1089,7 +1122,7 @@ static char* preprocessor_skip_whitespace_fast(char* s)
 
 // This is used when trimming macro arguments of trailing whitespace
 // @OPTIMIZE: turn this into a smart state-machine scanner?
-static char* preprocessor_skip_whitespace_reverse(char* s)
+static const char* preprocessor_skip_whitespace_reverse(const char* s)
 {
 	for (;;)
 	{
@@ -1106,7 +1139,7 @@ static char* preprocessor_skip_whitespace_reverse(char* s)
 	}
 }
 
-static char* scan_pp_identifier(char* p)
+static const char* scan_pp_identifier(const char* p)
 {
 	assert(char_is_pp_identifier_first(*p));
 	++p;
@@ -1119,8 +1152,8 @@ static char* scan_pp_identifier(char* p)
 // which will break column numbers in compiler, so this helper function recovers
 static void rewind_to_start_of_line(parse_state* cs)
 {
-	char* q = cs->src;
-	char* p = cs->src + cs->src_offset;
+	const char* q = cs->src;
+	const char* p = cs->src + cs->src_offset;
 
 	// stop at end of line or non-whitespace (which actually should only be able to be end of a multiline comment, i.e. '/')
 	while (p > q && char_is_nonnewline_whitespace(p[-1]))
@@ -1236,7 +1269,7 @@ static void output_line_directive(parse_state* cs);
 // characters in the file. We don't handle backslash concatenation, so no fast include-guard if so.
 static int scan_whitespace_and_comments(parse_state* cs)
 {
-	char* p = cs->src + cs->src_offset;
+	const char* p = cs->src + cs->src_offset;
 	int state = PP_WHITESCAN_STATE_whitespace;
 	int line_count = cs->src_line_number;
 
@@ -1274,7 +1307,7 @@ static int scan_whitespace_and_comments(parse_state* cs)
 static int copy_to_action_point(parse_state* cs)
 {
 	int state;
-	char* p = cs->src + cs->src_offset;
+	const char* p = cs->src + cs->src_offset;
 	char* out;
 	char* q;
 
@@ -1424,7 +1457,7 @@ static int copy_to_action_point(parse_state* cs)
 static int copy_to_action_point_macro_expansion(parse_state* cs)
 {
 	int state;
-	char* p = cs->src + cs->src_offset;
+	const char* p = cs->src + cs->src_offset;
 	char* out = cs->dest;
 	stb_assume(out != NULL);
 	char* q;
@@ -1468,7 +1501,7 @@ static int copy_to_action_point_macro_expansion(parse_state* cs)
 // when in a conditionally-disabled section of code, we want to
 // scan as fast as possible to the next directive; need a state-machine
 // to handle comments efficiently
-static char* scan_to_directive(char* p, int* p_line_number)
+static const char* scan_to_directive(const char* p, int* p_line_number)
 {
 	int line_number = *p_line_number;
 	int state = PP_STATE_ready, prev_state;
@@ -1508,7 +1541,7 @@ static char* scan_to_directive(char* p, int* p_line_number)
 //
 
 // this is used when parsing macro arguments... so @TODO comments could be present
-static char* preprocessor_skip_whitespace(char* s, int* newline_count)
+static const char* preprocessor_skip_whitespace(const char* s, int* newline_count)
 {
 	for (;;)
 	{
@@ -1538,7 +1571,7 @@ static char* preprocessor_skip_whitespace(char* s, int* newline_count)
 
 // this is used when parsing directives, so newlines terminate it
 // comments are allowed but rare.
-static char* preprocessor_skip_whitespace_in_directive(char* s, int* newline_count)
+static const char* preprocessor_skip_whitespace_in_directive(const char* s, int* newline_count)
 {
 	for (;;)
 	{
@@ -1582,7 +1615,7 @@ scan_to_end_of_line:
 	return s;
 }
 
-static char* copy_line_without_comments(char* buffer, size_t buffer_size, char* p, int* line_count, char** endp)
+static char* copy_line_without_comments(char* buffer, size_t buffer_size, const char* p, int* line_count, const char** endp)
 {
 	char *out = buffer, *end = buffer + buffer_size;
 	while (char_is_nonnewline_whitespace(*p))
@@ -1653,7 +1686,7 @@ static char* copy_line_without_comments(char* buffer, size_t buffer_size, char*
 	{
 		size_t length = out - buffer;
 		out = NULL;	 // stb_ds array
-		arrsetlen(out, length);
+		arrinitlen(out, length);
 		memcpy(out, buffer, length);
 
 		// this case is most likely to be a long #define
@@ -1713,7 +1746,6 @@ static char* copy_line_without_comments(char* buffer, size_t buffer_size, char*
 				{
 #pragma warning(push)
 #pragma warning(disable:6011) // MSVC analysis doesn't understand that arrmaybegrow can validly take null and allocate (thinks it's a possible null deref)
-
 					arrput(out, *p++);
 #pragma warning(pop)
 				}
@@ -1859,7 +1891,7 @@ static void preprocess_string(parse_state* cs, int in_macro_expansion, char* in_
 
 	if (in_macro_expansion == IN_MACRO_include_guard_scan)
 	{
-		char* p;
+		const char* p;
 		int action;
 
 		in_macro_expansion = IN_MACRO_no;
@@ -1981,7 +2013,7 @@ static void preprocess_string(parse_state* cs, int in_macro_expansion, char* in_
 					// so we now check to see whether it's followed by (. clang looks inside following macros
 					// for this, but VC6 doesn't, and it would require a rearchitecture to look inside the macros
 					// (and probably an efficiency hit), so we'll do it the VC6 way.
-					char* s = cs->src + cs->src_offset;
+					const char* s = cs->src + cs->src_offset;
 					// @TODO skip comments and newlines
 					while (char_is_nonnewline_whitespace(*s))
 						++s;
@@ -2021,7 +2053,16 @@ static void preprocess_string(parse_state* cs, int in_macro_expansion, char* in_
 
 			case PP_STATE_saw_hash:
 			{
-				if (do_error_code(cs, PP_RESULT_directive_not_at_start_of_line, "Preprocessor directive must appear at beginning of line."))
+				// special case for HLSL: need to handle the 1.#INF infinity constant and not error here.
+				int hlsl_inf = ((cs->src_offset >= 2) &&
+					((cs->src_offset + 3) < cs->src_length) &&
+					cs->src[cs->src_offset - 2] == '1' &&
+					cs->src[cs->src_offset - 1] == '.' &&
+					cs->src[cs->src_offset + 1] == 'I' &&
+					cs->src[cs->src_offset + 2] == 'N' &&
+					cs->src[cs->src_offset + 3] == 'F');
+
+				if (!hlsl_inf && do_error_code(cs, PP_RESULT_directive_not_at_start_of_line, "Preprocessor directive must appear at beginning of line."))
 					return;
 				else
 				{
@@ -2081,7 +2122,7 @@ static void update_macro_filter(pp_context* c, char* identifier)
 		c->macro_name_filter[(uint8)identifier[0]][(uint8)identifier[1]] = 1;
 }
 
-static struct macro_definition* create_macro_definition(parse_state* ps, char* def)
+static struct macro_definition* create_macro_definition(parse_state* ps, const char* def)
 {
 	// keep track of the state so we can share the diagnostic code
 	enum
@@ -2093,9 +2134,9 @@ static struct macro_definition* create_macro_definition(parse_state* ps, char* d
 
 	pp_context* c = ps->context;
 	char identifier_buffer[MAX_MACRO_PARAMETER_NAME_LENGTH + 1];
-	struct macro_definition* m = stb_arena_alloc(&c->macro_arena, sizeof(*m));
-	char* start = preprocessor_skip_whitespace_simple(def);
-	char* p;
+	struct macro_definition* m = (struct macro_definition*)stb_arena_alloc(&c->macro_arena, sizeof(*m));
+	const char* start = preprocessor_skip_whitespace_simple(def);
+	const char* p;
 
 	m->symbol_name = NULL;
 	m->disabled = 0;
@@ -2136,8 +2177,17 @@ static struct macro_definition* create_macro_definition(parse_state* ps, char* d
 		}
 
 		m->num_parameters = MACRO_NUM_PARAMETERS_no_parentheses;
-		m->simple_expansion = stb_arena_alloc_string_length(&c->macro_arena, start, p - start);
-		m->simple_expansion_length = (int)(p - start);
+		size_t expansion_length = p - start;
+		if (*(start + (expansion_length - 1)) == '>') // hacky special case in macro expansion to work around DXC bug (not handling >> in template declarations)
+		{
+			m->simple_expansion = arena_alloc_string_trailingspace(&c->macro_arena, start, expansion_length);
+			m->simple_expansion_length = (int)expansion_length + 1;
+		}
+		else
+		{
+			m->simple_expansion = stb_arena_alloc_string_length(&c->macro_arena, start, expansion_length);
+			m->simple_expansion_length = (int)(expansion_length);
+		}
 	}
 	else
 	{
@@ -2248,7 +2298,7 @@ static struct macro_definition* create_macro_definition(parse_state* ps, char* d
 			int stringify = 0;
 			int token_paste_after_text = 0;
 			expansion_part e;
-			char* end;
+			const char* end;
 			start = p;
 			// we need to parse to an identifier or '#', and we need to
 			// skip comments and string literals. we maybe could speed
@@ -2267,7 +2317,7 @@ static struct macro_definition* create_macro_definition(parse_state* ps, char* d
 				else if (char_is_pp_identifier_first(*p))
 				{
 					size_t len;
-					char* s = scan_pp_identifier(p);
+					const char* s = scan_pp_identifier(p);
 					len = s - p;
 					found_parameter = ((int)(size_t)stringhash_get(&parameters, p, len)) - 1;
 					if (found_parameter >= 0)
@@ -2326,7 +2376,7 @@ static struct macro_definition* create_macro_definition(parse_state* ps, char* d
 			if (found_parameter < 0)
 			{
 				size_t len;
-				char* s;
+				const char* s;
 				assert(require_identifier);
 				p = preprocessor_skip_whitespace_simple(p);
 				if (!char_is_pp_identifier_first(*p))
@@ -2491,10 +2541,10 @@ static char *parse_argument(char *text, int *line_number, int *needs_copying)
 }
 #endif
 
-static char* copy_argument(char* text, int* line_number, char** p_out)
+static const char* copy_argument(const char* text, int* line_number, char** p_out)
 {
 	int num_parens = 0;
-	char* p = text;
+	const char* p = text;
 	char* out = *p_out;	 // need to be able to append to existing buffer for VA_ARGS
 	arrsetcap(out, 200);
 
@@ -2502,7 +2552,7 @@ static char* copy_argument(char* text, int* line_number, char** p_out)
 	{
 		// scan characters that don't need processing
 		// we don't have a length bound, so we don't copy as we go, we make a second pass
-		char* q = p;
+		const char* q = p;
 		ptrdiff_t addlen;
 		size_t oldlen;
 		while (!char_is_argument_parse_action((uint8)*p))
@@ -2636,9 +2686,9 @@ static void maybe_expand_macro(parse_state* cs, struct macro_definition* pending
 {
 	pp_context* c = cs->context;
 	struct macro_definition* md;
-	char* identifier = cs->fast_dest ? cs->fast_dest : (char*)cs->dest + arrlen(cs->dest);
+	char* identifier = cs->fast_dest ? cs->fast_dest : (char*)cs->dest + arrlennonull(cs->dest);
 	size_t identifier_length;
-	char* p = cs->src + cs->src_offset;
+	const char* p = cs->src + cs->src_offset;
 	char* q = identifier;
 	uint8 c0 = (uint8)p[0];
 	uint8 c1 = (uint8)p[1];
@@ -2650,7 +2700,7 @@ static void maybe_expand_macro(parse_state* cs, struct macro_definition* pending
 	}
 	else
 	{
-		char* start = p;
+		const char* start = p;
 
 #ifndef DISABLE_SHORT_MACRO_FILTER
 		if (!char_is_pp_identifier(c1))
@@ -2684,7 +2734,7 @@ static void maybe_expand_macro(parse_state* cs, struct macro_definition* pending
 			*q++ = *p++;
 
 		identifier_length = p - start;
-		md = stringhash_get(&c->macro_map, identifier, identifier_length);
+		md = (struct macro_definition*)stringhash_get(&c->macro_map, identifier, identifier_length);
 		if (md == NULL || md->disabled)
 		{
 		not_macro:
@@ -2853,7 +2903,7 @@ static void maybe_expand_macro(parse_state* cs, struct macro_definition* pending
 		int i;
 		int temp_lines;
 		int parse_rest = 0;
-		char* s = p;
+		const char* s = p;
 
 		// plan:
 		//   1. find opening parenthesis, if any
@@ -2950,7 +3000,7 @@ static void maybe_expand_macro(parse_state* cs, struct macro_definition* pending
 			}
 
 			{
-				char* e = s;
+				char* e;
 				// need to strip leading and trailing whitespace so token-pasting works
 
 				arrput(copy, 0);
@@ -2959,10 +3009,8 @@ static void maybe_expand_macro(parse_state* cs, struct macro_definition* pending
 
 				if (e - 1 > p)
 				{
-					e = preprocessor_skip_whitespace_reverse(e);
+					e = (char*)preprocessor_skip_whitespace_reverse(e);
 					assert(e > p);
-					if (e < p)
-						e = p;	// this should be impossible
 				}
 				*e = 0;
 
@@ -2974,6 +3022,12 @@ static void maybe_expand_macro(parse_state* cs, struct macro_definition* pending
 			}
 		}
 
+		if (md->num_parameters == 0)
+		{
+			assert(*p == ')');
+			s = p;
+		}
+
 		if (c->error)
 		{
 			arrfree(arguments);
@@ -3026,16 +3080,16 @@ static void maybe_expand_macro(parse_state* cs, struct macro_definition* pending
 
 				arrput(copy, 0);
 				p = copy;
-				e = p + arrlen(p);
+				e = copy + arrlen(copy);
 
 				// strip trailing whitespace
 				while (char_is_whitespace(e[-1]))
 					--e;
-				if (e < p)
-					e = p;
+				if (e < copy)
+					e = copy;
 				*e = 0;
 
-				arrput(arguments, p);
+				arrput(arguments, copy);
 				cs->src_line_number += arg_newlines;
 			}
 		}
@@ -3064,20 +3118,20 @@ static void maybe_expand_macro(parse_state* cs, struct macro_definition* pending
 		acs.dest = acs.fast_dest = 0;
 
 		// populate temporary buffer
-		s = 0;
+		char* tmp = 0;
 		for (i = 0; i < arrlen(md->expansion) - 1; ++i)
 		{
 			int an;
 			size_t off;
 
-			off = arraddnindex(s, md->expansion[i].text_length);
+			off = arraddnindex(tmp, md->expansion[i].text_length);
 			// copy macro-expansion text
-			memcpy(s + off, md->expansion[i].text, md->expansion[i].text_length);
+			memcpy(tmp + off, md->expansion[i].text, md->expansion[i].text_length);
 
 			an = md->expansion[i].argument_index;
 
 			// sanity check added to satisfy static analysis; this shouldn't be possible
-			if (an >= arrlen(arguments))
+			if (arguments == NULL || an >= arrlen(arguments))
 			{
 				do_error(cs, "Internal error; arguments array not constructed as expected");
 				return;
@@ -3094,7 +3148,7 @@ static void maybe_expand_macro(parse_state* cs, struct macro_definition* pending
 				if (arguments[an] == NULL)
 				{
 					// find the generated comma and set it to whitespace
-					char* t = s + arrlen(s);
+					char* t = tmp + arrlen(tmp);
 					while (*--t != ',')
 						;
 					*t = ' ';
@@ -3107,17 +3161,17 @@ static void maybe_expand_macro(parse_state* cs, struct macro_definition* pending
 			{
 				if (md->expansion[i].stringify_argument)
 				{
-					arrput(s, ' ');
-					arrput(s, '"');
+					arrput(tmp, ' ');
+					arrput(tmp, '"');
 					// @TODO: stringify needs to insert backslashes before \ and "
-					off = arraddnindex(s, arrlen(arguments[an]));
-					memcpy(s + off, arguments[an], arrlen(arguments[an]));
-					arrput(s, '"');
+					off = arraddnindex(tmp, arrlen(arguments[an]));
+					memcpy(tmp + off, arguments[an], arrlen(arguments[an]));
+					arrput(tmp, '"');
 				}
 				else
 				{
-					off = arraddnindex(s, arrlen(arguments[an]));
-					memcpy(s + off, arguments[an], arrlen(arguments[an]));
+					off = arraddnindex(tmp, arrlen(arguments[an]));
+					memcpy(tmp + off, arguments[an], arrlen(arguments[an]));
 				}
 			}
 			else
@@ -3125,17 +3179,17 @@ static void maybe_expand_macro(parse_state* cs, struct macro_definition* pending
 				struct macro_definition* ffm = 0;
 				acs.src = arguments[an];
 				acs.src_length = arrlen(arguments[an]);
-				acs.dest = s;
+				acs.dest = tmp;
 				acs.fast_dest = 0;
 				acs.src_offset = 0;
 				preprocess_string(&acs, in_macro_expansion, identifier, &ffm);
-				s = acs.dest;
+				tmp = acs.dest;
 				if (ffm)
 				{
 					// concatenate it into the result string... if it's at the end, we'll catch it again in the processing below
 					size_t len = strlen(ffm->symbol_name);
-					off = arraddnindex(s, len);
-					memcpy(s + off, ffm->symbol_name, len);
+					off = arraddnindex(tmp, len);
+					memcpy(tmp + off, ffm->symbol_name, len);
 				}
 				if (cs->context->stop)
 					return;
@@ -3143,11 +3197,11 @@ static void maybe_expand_macro(parse_state* cs, struct macro_definition* pending
 		}
 
 		{
-			size_t off = arraddnindex(s, md->expansion[i].text_length);
-			memcpy(s + off, md->expansion[i].text, md->expansion[i].text_length);
+			size_t off = arraddnindex(tmp, md->expansion[i].text_length);
+			memcpy(tmp + off, md->expansion[i].text, md->expansion[i].text_length);
 		}
 
-		arrput(s, 0);
+		arrput(tmp, 0);
 
 		{
 			for (i = 0; i < arrlen(arguments); ++i)
@@ -3155,8 +3209,8 @@ static void maybe_expand_macro(parse_state* cs, struct macro_definition* pending
 			arrfree(arguments);
 		}
 
-		ncs.src = s;
-		ncs.src_length = arrlen(s);
+		ncs.src = tmp;
+		ncs.src_length = arrlen(tmp);
 		ncs.src_offset = 0;
 
 		//////////////////////////////////////////////////////////////////
@@ -3180,7 +3234,7 @@ static void maybe_expand_macro(parse_state* cs, struct macro_definition* pending
 //
 
 // it's already been cleaned of backslash-continuation and is NUL-terminated
-static char* macro_expand_directive(parse_state* cs, char* p, char* directive)
+static char* macro_expand_directive(parse_state* cs, const char* p, char* directive)
 {
 	parse_state ncs;
 	ncs = *cs;
@@ -3193,7 +3247,7 @@ static char* macro_expand_directive(parse_state* cs, char* p, char* directive)
 	return ncs.dest;
 }
 
-static int evaluate_if(parse_state* cs, char* p, int* syntax_error)
+static int evaluate_if(parse_state* cs, const char* p, int* syntax_error)
 {
 	int result = 0;
 	parse_state ncs;
@@ -3220,10 +3274,10 @@ static int evaluate_if(parse_state* cs, char* p, int* syntax_error)
 	return result;
 }
 
-static char* parse_pp_identifier_in_directive(parse_state* ps, char* text, char* identifier, size_t identifier_buffer_length, char* directive, size_t* idlen)
+static const char* parse_pp_identifier_in_directive(parse_state* ps, const char* text, char* identifier, size_t identifier_buffer_length, char* directive, size_t* idlen)
 {
 	size_t i;
-	char* p = preprocessor_skip_whitespace_simple(text);
+	const char* p = preprocessor_skip_whitespace_simple(text);
 	*idlen = 0;
 	identifier[0] = 0;
 	if (!char_is_pp_identifier_first(*p))
@@ -3266,11 +3320,11 @@ enum
 };
 
 // normally this comes from a safe buffer, but during disabled-by-#if scanning it's an uncopied directive
-static int parse_directive_after_hash(char** ptr_p, int* newlines)
+static int parse_directive_after_hash(const char** ptr_p, int* newlines)
 {
 	int sum = 0;
-	char* p = *ptr_p;
-	char* start;
+	const char* p = *ptr_p;
+	const char* start;
 	p = preprocessor_skip_whitespace_in_directive(p, newlines);
 	start = p;
 	while (char_is_pp_identifier_first(*p))
@@ -3296,18 +3350,18 @@ enum
 	RETURN_BEHAVIOR_skip_to_end_of_line,
 };
 
-static char* preprocess_string_from_file(parse_state* ps, char* filename, char* text, size_t textlen, int conditional_compilation_nesting, int is_header);
-static char* resolve_file(char path[4096],
+static char* preprocess_string_from_file(parse_state* ps, const char* filename, const char* text, size_t textlen, int conditional_compilation_nesting, int is_header);
+static const char* resolve_file(char path[MAX_FILENAME],
 						  parse_state* ps,
 						  ppbool check_source_dirs,
 						  ppbool check_normal_includes,
 						  ppbool check_system_includes,
-						  char* filename,
+						  const char* filename,
 						  ptrdiff_t filename_len);
-static char* load_file_with_nul_terminator(char* filename, size_t* ptr_len, int* mode);
-static void load_file_free(char* filedata, int mode);
+static const char* load_file_with_nul_terminator(const char* filename, size_t* ptr_len, int* mode, void* custom_context);
+static void load_file_free(const char* filename, const char* filedata, int mode, void* custom_context);
 
-static void pragma_once(pp_context* c, char* filename, char* macroname)
+static void pragma_once(pp_context* c, const char* filename, char* macroname)
 {
 	if (macroname)
 		macroname = stb_arena_alloc_string(&c->macro_arena, macroname);
@@ -3316,11 +3370,12 @@ static void pragma_once(pp_context* c, char* filename, char* macroname)
 	shput(c->once_map, filename, macroname);
 }
 
-static int process_include(parse_state* cs, char* start, conditional_state* cons)
+
+static int process_include(parse_state* cs, const char* start, conditional_state* cons)
 {
 	int ret = RETURN_BEHAVIOR_break;
 	pp_context* c = cs->context;
-	char *filename, *p;
+	const char *filename, *p;
 	p = preprocessor_skip_whitespace_simple(start);
 	// @TODO: macro expansion
 	if (!(*p == '<' || *p == '"'))
@@ -3332,8 +3387,8 @@ static int process_include(parse_state* cs, char* start, conditional_state* cons
 	else
 	{
 		char* once;
-		char path[4096];
-		char *start_char = p + 1, *old_p;
+		char path[MAX_FILENAME];
+		const char *start_char = p + 1, *old_p;
 		char end_char = (*p == '<') ? '>' : '"';
 		++p;
 		while (*p != end_char)
@@ -3387,15 +3442,16 @@ static int process_include(parse_state* cs, char* start, conditional_state* cons
 			struct macro_definition* md = NULL;
 
 			if (once != NULL)
-				md = stringhash_get(&c->macro_map, once, strlen(once));
+				md = (struct macro_definition*)stringhash_get(&c->macro_map, once, strlen(once));
 
 			// if macro is defined, it's an include guard, which means if it's defined we can skip processing
 			if (md == NULL)
 			{
 				int mode;
-				char *data, *result;
+				const char* data;
+				char *result;
 				size_t len;
-				data = load_file_with_nul_terminator(filename, &len, &mode);
+				data = load_file_with_nul_terminator(filename, &len, &mode, c->custom_context);
 
 				if (data == NULL)
 				{
@@ -3405,7 +3461,7 @@ static int process_include(parse_state* cs, char* start, conditional_state* cons
 
 				c->num_includes += 1;
 				result = preprocess_string_from_file(cs, filename, data, len, cons->conditional_compilation_nesting_level, 1);
-				load_file_free(data, mode);
+				load_file_free(filename, data, mode, c->custom_context);
 				if (cs->fast_dest)
 					cs->fast_dest = result;
 				else
@@ -3427,8 +3483,8 @@ static void process_directive(parse_state* cs, conditional_state* cons)
 {
 	int n;
 	pp_context* c = cs->context;
-	char *p, *endptr;
-	char* alloc = NULL;	 // all exit paths SHOULD do 'if (alloc) arrfree(alloc)', but we allow it to leak on error-handling returns
+	const char *p, *endptr;
+	const char* alloc = NULL;	 // all exit paths SHOULD do 'if (alloc) arrfree(alloc)', but we allow it to leak on error-handling returns
 	char buffer[1000];
 	int dummy;	// uninitialized, should never get incremented by parse_directive_after_hash
 
@@ -3501,15 +3557,14 @@ static void process_directive(parse_state* cs, conditional_state* cons)
 							static char* directives[] = {
 								"define", "else", "elif", "endif", "error", "if", "ifdef", "ifndef", "include", "line", "pragma", "undef",
 							};
-							char identifier[64], *match = 0;
+							char identifier[64] = { 0 }, *match = 0;
 							int i;
-							for (i = 0; i < sizeof(identifier) - 1; ++i)
+							for (i = 0; i < sizeof(identifier); ++i)
 							{
 								if (!char_is_pp_identifier(p[i]))
 									break;
 								identifier[i] = p[i];
 							}
-							identifier[i++] = 0;
 							if (i < sizeof(identifier) && pp_result_mode[PP_RESULT_unrecognized_directive] <= did_you_mean_threshold)
 							{
 								match = find_closest_matching_string(
@@ -3544,8 +3599,10 @@ static void process_directive(parse_state* cs, conditional_state* cons)
 			p = preprocessor_skip_whitespace_simple(p);
 			(void)parse_pp_identifier_in_directive(cs, p, identifier, sizeof(identifier), "pragma", &idlen);
 			// check for preprocessor tokens we recognize
-			switch (identifier[0])
+			if (idlen >= 3) // don't bother comparing if the identifier is < 3 chars long (to satisfy some CA failure conditions)
 			{
+				switch (identifier[0])
+				{
 				case 'o':
 					if (0 == strcmp(identifier, "once"))
 					{
@@ -3578,6 +3635,7 @@ static void process_directive(parse_state* cs, conditional_state* cons)
 						}
 					}
 					break;
+				}
 			}
 
 			// copy pragma through to output... existing code will output the newline, so don't do it here
@@ -3617,7 +3675,7 @@ static void process_directive(parse_state* cs, conditional_state* cons)
 
 			if (m != NULL)
 			{
-				old = stringhash_get(&c->macro_map, m->symbol_name, m->symbol_name_length);
+				old = (struct macro_definition*)stringhash_get(&c->macro_map, m->symbol_name, m->symbol_name_length);
 				if (old != NULL)
 				{
 					// @TODO: compare new definition to old and produce error/warning
@@ -3645,7 +3703,7 @@ static void process_directive(parse_state* cs, conditional_state* cons)
 				if (alloc)
 					arrfree(alloc);
 
-			md = stringhash_get(&c->macro_map, identifier, idlen);
+			md = (struct macro_definition*)stringhash_get(&c->macro_map, identifier, idlen);
 
 			if (md != NULL)
 			{
@@ -3657,7 +3715,7 @@ static void process_directive(parse_state* cs, conditional_state* cons)
 				else
 				{
 					stringhash_delete(&c->macro_map, identifier, idlen);
-					shput(c->undef_map, identifier, (void*)(size_t)1);
+					shput(c->undef_map, identifier, (struct macro_definition*)(size_t)1);
 				}
 			}
 			else
@@ -3758,7 +3816,7 @@ static void process_directive(parse_state* cs, conditional_state* cons)
 			arrpush(c->ifdef_stack, ifinfo);
 			++cons->conditional_compilation_nesting_level;
 
-			m = stringhash_get(&c->macro_map, identifier, idlen);
+			m = (struct macro_definition*)stringhash_get(&c->macro_map, identifier, idlen);
 			if (m == NULL || identifier[0] == 0)
 				cons->conditionally_disable = CONDITIONAL_skip_current_block;
 			else
@@ -3793,7 +3851,7 @@ static void process_directive(parse_state* cs, conditional_state* cons)
 			arrpush(c->ifdef_stack, ifinfo);
 			++cons->conditional_compilation_nesting_level;
 
-			m = stringhash_get(&c->macro_map, identifier, idlen);
+			m = (struct macro_definition*)stringhash_get(&c->macro_map, identifier, idlen);
 			if (m != NULL)
 				cons->conditionally_disable = CONDITIONAL_skip_current_block;
 			else
@@ -4002,8 +4060,8 @@ static void process_directive(parse_state* cs, conditional_state* cons)
 
 		for (;;)
 		{
-			char* q;
-			char* hash_location;
+			const char* q;
+			const char* hash_location;
 			p = scan_to_directive(p, &newlines);  // fast scanner that skips non-directive content
 
 			if (*p == 0)
@@ -4130,7 +4188,7 @@ static void process_directive(parse_state* cs, conditional_state* cons)
 }
 
 int preprocessor_automatic_include_guard_detection = 1;
-static char* preprocess_string_from_file(parse_state* ps, char* filename, char* text, size_t textlen, int conditional_compilation_nesting, int is_header)
+static char* preprocess_string_from_file(parse_state* ps, const char* filename, const char* text, size_t textlen, int conditional_compilation_nesting, int is_header)
 {
 	parse_state cs = *ps;
 	cs.parent = ps;
@@ -4294,7 +4352,7 @@ void init_preprocessor_scanner(void)
 	transition(STATE_DEFAULT, PP_CHAR_CLASS_whitespace, PP_STATE_ready);
 	transition(STATE_DEFAULT, PP_CHAR_CLASS_cr, PP_STATE_saw_cr);
 	transition(STATE_DEFAULT, PP_CHAR_CLASS_lf, PP_STATE_saw_lf);
-	transition(STATE_DEFAULT, PP_CHAR_CLASS_period, PP_STATE_in_number);
+	transition(STATE_DEFAULT, PP_CHAR_CLASS_period, PP_STATE_saw_period);
 	transition(STATE_DEFAULT, PP_CHAR_CLASS_digit, PP_STATE_in_number);
 	transition(STATE_DEFAULT, PP_CHAR_CLASS_idtype, PP_STATE_saw_identifier_start);
 	transition(STATE_DEFAULT, PP_CHAR_CLASS_idtype_e_E, PP_STATE_saw_identifier_start);
@@ -4347,6 +4405,9 @@ void init_preprocessor_scanner(void)
 	transition(PP_STATE_comment_multiline_seen_lf, PP_CHAR_CLASS_cr, PP_STATE_comment_multiline);  // don't double-count LF-CR
 #endif
 
+	transition(PP_STATE_saw_period, PP_CHAR_CLASS_digit, PP_STATE_in_number);
+	transition(PP_STATE_saw_period, PP_CHAR_CLASS_idtype, PP_STATE_saw_identifier_start);
+
 	// parse preprocessor number
 
 	transition(PP_STATE_in_number, PP_CHAR_CLASS_digit, PP_STATE_in_number);
@@ -4419,7 +4480,8 @@ void init_preprocessor_scanner(void)
 	pp_initscan_state_is_end_of_line[PP_WHITESCAN_STATE_comment_multiline_seen_lf] = 1;
 }
 
-struct macro_definition* pp_define(struct stb_arena* a, char* p)
+
+struct macro_definition* pp_define(struct stb_arena* a, const char* p)
 {
 	struct macro_definition* m;
 	pp_context c = {0};
@@ -4441,23 +4503,24 @@ static void define_macro(pphash* map, struct macro_definition* m)
 }
 
 char* preprocess_file(char* output_autobuffer,
-					  char* filename,
-					  char** include_paths,
-					  int num_include_paths,
-					  char** sys_include_paths,
-					  int num_sys_include_paths,
-					  struct macro_definition** predefined_macros,
-					  int num_predefined_macros,
-					  pp_diagnostic** pd,
-					  int* num_pd)
+	const char* filename,
+	void* custom_context,
+	char** include_paths,
+	int num_include_paths,
+	char** sys_include_paths,
+	int num_sys_include_paths,
+	struct macro_definition** predefined_macros,
+	int num_predefined_macros,
+	pp_diagnostic** pd,
+	int* num_pd)
 {
 	char* output = 0;
-	pp_context c = {0};
+	pp_context c = { 0 };
 	int main_mode, i;
 	size_t length;
-	char* main_file;
+	const char* main_file;
 
-	main_file = load_file_with_nul_terminator(filename, &length, &main_mode);
+	main_file = load_file_with_nul_terminator(filename, &length, &main_mode, custom_context);
 	*num_pd = 0;
 
 	if (main_file == 0)
@@ -4476,6 +4539,7 @@ char* preprocess_file(char* output_autobuffer,
 	c.sys_include_paths = sys_include_paths;
 	c.num_include_paths = num_include_paths;
 	c.num_sys_include_paths = num_sys_include_paths;
+	c.custom_context = custom_context;
 
 	// allocate the macro definition dictionary very large so (a) it's unlikely to grow,
 	// and (b) so most things hit in there first slot. but going too much bigger can
@@ -4518,7 +4582,7 @@ char* preprocess_file(char* output_autobuffer,
 		}
 	}
 
-	load_file_free(main_file, main_mode);
+	load_file_free(filename, main_file, main_mode, custom_context);
 
 	shfree(c.once_map);
 	stringhash_destroy(&c.macro_map);
@@ -4621,7 +4685,10 @@ static void init_char_type(void)
 	pp_char_flags2['/'] |= CHAR2_IS_SLASH;
 }
 
-void init_preprocessor(void)
+void init_preprocessor(
+	loadfile_callback_func load_callback, 
+	freefile_callback_func free_callback,
+	resolveinclude_callback_func resolve_callback)
 {
 #ifdef HASHTEST
 	test_stringhash();
@@ -4643,6 +4710,10 @@ void init_preprocessor(void)
 	init_directive("line", HASH_line);
 	init_directive("pragma", HASH_pragma);
 	init_directive("undef", HASH_undef);
+
+	loadfile_callback = load_callback;
+	freefile_callback = free_callback;
+	resolveinclude_callback = resolve_callback;
 }
 
 // File reader
@@ -4658,12 +4729,13 @@ typedef struct
 
 typedef struct
 {
-	char* key;
-	char* value;
+	const char* key;
+	const char* value;
 } resolved_header;
 
 static cached_header* cached_header_map;
 static size_t cached_header_total_size;
+static resolved_header* resolved_callback_header_map;
 static resolved_header* resolved_header_map;
 static resolved_header* resolved_system_header_map;
 static struct stb_arena resolve_arena;
@@ -4697,7 +4769,7 @@ void init_preprocessor_cache(void)
 	sh_new_arena(resolved_system_header_map);
 }
 
-static int find_file_path(char path[4096], char* filename, char** path_list, ptrdiff_t path_list_len)
+static int find_file_path(char path[MAX_FILENAME], char* filename, char** path_list, ptrdiff_t path_list_len)
 {
 	int i;
 	FILE* f = NULL;
@@ -4716,7 +4788,7 @@ static int find_file_path(char path[4096], char* filename, char** path_list, ptr
 	return 0;
 }
 
-static char* my_strcpy(char* dest, char* src)
+static char* my_strcpy(char* dest, const char* src)
 {
 	while (*src)
 		*dest++ = *src++;
@@ -4724,20 +4796,34 @@ static char* my_strcpy(char* dest, char* src)
 	return dest;
 }
 
-static char* resolve_file(char path[4096],
+static const char* resolve_file(char path[MAX_FILENAME],
 						  parse_state* ps,
 						  ppbool check_source_dirs,
 						  ppbool check_normal_includes,
 						  ppbool check_system_includes,
-						  char* filename,
+						  const char* filename,
 						  ptrdiff_t filename_len)
 {
 	pp_context* c = ps->context;
 
-	char filenamez[4096];
+	char filenamez[MAX_FILENAME];
 	strncpy(filenamez, filename, filename_len);
 	filenamez[filename_len] = 0;
 
+	if (resolveinclude_callback != NULL)
+	{
+		const char* resolved = shget(resolved_callback_header_map, path);
+		if (resolved)
+			return resolved;
+
+		resolved  = (*resolveinclude_callback)(filename, filename_len, ps->filename, &resolve_arena, c->custom_context);
+		if (resolved)
+		{
+			shput(resolved_callback_header_map, filenamez, resolved);
+			return resolved;
+		}
+	}
+
 	if (check_source_dirs)
 	{
 		do
@@ -4763,7 +4849,7 @@ static char* resolve_file(char path[4096],
 
 	if (check_normal_includes)
 	{
-		char* resolved = shget(resolved_header_map, filenamez);
+		const char* resolved = shget(resolved_header_map, filenamez);
 		if (resolved)
 			return resolved;
 
@@ -4777,7 +4863,7 @@ static char* resolve_file(char path[4096],
 
 	if (check_system_includes)
 	{
-		char* resolved = shget(resolved_system_header_map, filenamez);
+		const char* resolved = shget(resolved_system_header_map, filenamez);
 		if (resolved)
 			return resolved;
 
@@ -4792,36 +4878,52 @@ static char* resolve_file(char path[4096],
 	return NULL;
 }
 
+#define LOAD_FILE_MODE_callback -3
 #define LOAD_FILE_MODE_malloced -2
 #define LOAD_FILE_MODE_memmappend -1
 #define LOAD_FILE_MODE_multi_memmaped 0	 // actual positive value is number of 4K pages past end of pointer that second mapping occurs
 
-static char* load_file_with_nul_terminator(char* filename, size_t* ptr_len, int* mode)
+static const char* load_file_with_nul_terminator(const char* filename, size_t* ptr_len, int* mode, void* custom_context)
 {
-	char* p;
-	size_t len;
-	FILE* f = fopen(filename, "rb");
-	*mode = LOAD_FILE_MODE_malloced;
-	if (f == 0)
-		return NULL;
-	fseek(f, 0, SEEK_END);
-	len = ftell(f);
-	fseek(f, 0, SEEK_SET);
-	p = malloc(len + 1);
-	if (p == NULL)
-		return NULL;
+	if (loadfile_callback != NULL)
+	{
+		*mode = LOAD_FILE_MODE_callback;
+		return (*loadfile_callback)(filename, custom_context, ptr_len);
+	}
+	else
+	{
+		char* p;
+		size_t len;
+		FILE* f = fopen(filename, "rb");
+		*mode = LOAD_FILE_MODE_malloced;
+		if (f == 0)
+			return NULL;
+		fseek(f, 0, SEEK_END);
+		len = ftell(f);
+		fseek(f, 0, SEEK_SET);
+		p = (char*)malloc(len + 1);
+		if (p == NULL)
+			return NULL;
 
-	p[len] = 0;
-	fread(p, 1, len, f);
-	fclose(f);
-	if (ptr_len)
-		*ptr_len = len + 1;
-	return p;
+		p[len] = 0;
+		fread(p, 1, len, f);
+		fclose(f);
+		if (ptr_len)
+			*ptr_len = len + 1;
+		return p;
+	}
 }
 
-static void load_file_free(char* filedata, int mode)
+static void load_file_free(const char* filename, const char* filedata, int mode, void* custom_context)
 {
-	free(filedata);
+	if (freefile_callback != NULL && mode == LOAD_FILE_MODE_callback)
+	{
+		(*freefile_callback)(filename, filedata, custom_context);
+	}
+	else if (mode == LOAD_FILE_MODE_malloced)
+	{
+		free((void*)filedata);
+	}
 }
 
 #if 0
diff --git a/Engine/Source/Developer/ShaderPreprocessor/Private/stb_preprocess/preprocessor.h b/Engine/Source/Developer/ShaderPreprocessor/Private/stb_preprocess/preprocessor.h
--- a/Engine/Source/Developer/ShaderPreprocessor/Private/stb_preprocess/preprocessor.h
+++ b/Engine/Source/Developer/ShaderPreprocessor/Private/stb_preprocess/preprocessor.h
@@ -54,18 +54,46 @@ typedef struct
 	pp_where* where;
 } pp_diagnostic;
 
-extern void preprocessor_test(void);
-extern void init_preprocessor(void);
-extern char* preprocess_file(char* output_storage,
-							 char* filename,
-							 char** include_paths,
-							 int num_include_paths,
-							 char** sys_include_paths,
-							 int num_sys_include_paths,
-							 struct macro_definition** predefined_macros,
-							 int num_predefined_macros,
-							 pp_diagnostic** pd,
-							 int* num_pd);
-extern void preprocessor_file_free(char* text, pp_diagnostic* pd);
-extern struct macro_definition* pp_define(struct stb_arena* a, char* p);
-extern unsigned int preprocessor_hash(char* data, size_t len);	// hash function for string of length 'len', roughly xxhash32
+
+#ifdef __cplusplus
+#define STB_PP_DEF extern "C"
+#else
+#define STB_PP_DEF extern
+#endif
+
+// callback function can be optionally specified to init_preprocessor to override file loading functionality
+// it's probably necessary to provide the freefile callback as well if you provide this (unless you are ok with
+// the default behaviour of calling "free" on the returned pointer)
+typedef const char* (*loadfile_callback_func)(const char* filename, void* custom_context, size_t* out_length);
+
+// callback function can be optionally specified to init_preprocessor to override freeing loaded files
+// it's probably necessary to provide the loadfile callback as well if you provide this
+typedef void (*freefile_callback_func)(const char* filename, const char* loaded_file, void* custom_context);
+
+// callback function can be optionally specified to init_preprocessor to override include resolution; in the
+// current incarnation this will still attempt file-based loading using the default behaviour in the event
+// that this callback returns null for a given file path.
+// note: returned resolved paths are expected to be allocated for the lifetime of preprocessor execution; 
+// for simplicity use stb_arena_alloc_string using the given "path_arena" stb_arena
+typedef const char* (*resolveinclude_callback_func)(const char* path, unsigned int path_len, const char* parent, struct stb_arena* path_arena, void* custom_context);
+
+STB_PP_DEF void preprocessor_test(void);
+STB_PP_DEF void init_preprocessor(
+	loadfile_callback_func load_callback,
+	freefile_callback_func free_callback,
+	resolveinclude_callback_func resolve_callback);
+STB_PP_DEF char* preprocess_file(
+	char* output_storage,
+	const char* filename,
+	void* custom_context,
+	char** include_paths,
+	int num_include_paths,
+	char** sys_include_paths,
+	int num_sys_include_paths,
+	struct macro_definition** predefined_macros,
+	int num_predefined_macros,
+	pp_diagnostic** pd,
+	int* num_pd);
+STB_PP_DEF void preprocessor_file_free(char* text, pp_diagnostic* pd);
+STB_PP_DEF struct macro_definition* pp_define(struct stb_arena* a, const char* p);
+STB_PP_DEF unsigned int preprocessor_hash(const char* data, size_t len);	// hash function for string of length 'len', roughly xxhash32
diff --git a/Engine/Source/Developer/ShaderPreprocessor/Private/stb_preprocess/stb_alloc.h b/Engine/Source/Developer/ShaderPreprocessor/Private/stb_preprocess/stb_alloc.h
--- a/Engine/Source/Developer/ShaderPreprocessor/Private/stb_preprocess/stb_alloc.h
+++ b/Engine/Source/Developer/ShaderPreprocessor/Private/stb_preprocess/stb_alloc.h
@@ -21,8 +21,8 @@ STB_ALLOC_DEF void* stb_arena_alloc(struct stb_arena* a, size_t size);
 STB_ALLOC_DEF void* stb_arena_alloc_aligned(struct stb_arena* a, size_t size, size_t align);
 
 // allocate a string from an arena
-STB_ALLOC_DEF char* stb_arena_alloc_string(struct stb_arena* a, char* str);
-STB_ALLOC_DEF char* stb_arena_alloc_string_length(struct stb_arena* a, char* str, size_t length);
+STB_ALLOC_DEF char* stb_arena_alloc_string(struct stb_arena* a, const char* str);
+STB_ALLOC_DEF char* stb_arena_alloc_string_length(struct stb_arena* a, const char* str, size_t length);
 
 // free the entire arena, leaves arena reset so it can be used again
 STB_ALLOC_DEF void stb_arena_free(struct stb_arena* a);
@@ -68,7 +68,7 @@ void* stb_arena_alloc_aligned(struct stb_arena* a, size_t size, size_t align)
 		// grow the list of blocks if necessary
 		if (a->num_blocks == 0)
 		{
-			a->blocks = malloc(sizeof(a->blocks[0]) * 32);
+			a->blocks = (void**)malloc(sizeof(a->blocks[0]) * 32);
 		}
 		else if (a->num_blocks >= 32)
 		{
@@ -78,7 +78,7 @@ void* stb_arena_alloc_aligned(struct stb_arena* a, size_t size, size_t align)
 				// grow table
 				void* ptr = realloc(a->blocks, sizeof(a->blocks[0]) * a->num_blocks * 2);
 				if (ptr)
-					a->blocks = ptr;
+					a->blocks = (void**)ptr;
 			}
 		}
 
@@ -127,7 +127,7 @@ void* stb_arena_alloc(struct stb_arena* a, size_t size)
 	return stb_arena_alloc_aligned(a, size, stb_arena_align_size(size));
 }
 
-char* stb_arena_alloc_string_length(struct stb_arena* a, char* str, size_t length)
+char* stb_arena_alloc_string_length(struct stb_arena* a, const char* str, size_t length)
 {
 	char* p = (char*)stb_arena_alloc_aligned(a, length + 1, 4);
 	memcpy(p, str, length);
@@ -135,7 +135,7 @@ char* stb_arena_alloc_string_length(struct stb_arena* a, char* str, size_t lengt
 	return p;
 }
 
-char* stb_arena_alloc_string(struct stb_arena* a, char* str)
+char* stb_arena_alloc_string(struct stb_arena* a, const char* str)
 {
 	return stb_arena_alloc_string_length(a, str, strlen(str));
 }
diff --git a/Engine/Source/Developer/ShaderPreprocessor/Private/stb_preprocess/stb_ds.h b/Engine/Source/Developer/ShaderPreprocessor/Private/stb_preprocess/stb_ds.h
--- a/Engine/Source/Developer/ShaderPreprocessor/Private/stb_preprocess/stb_ds.h
+++ b/Engine/Source/Developer/ShaderPreprocessor/Private/stb_preprocess/stb_ds.h
@@ -398,6 +398,7 @@ CREDITS
 
 #ifndef STBDS_NO_SHORT_NAMES
 #define arrlen stbds_arrlen
+#define arrlennonull stbds_arrlennonull
 #define arrlenu stbds_arrlenu
 #define arrput stbds_arrput
 #define arrpush stbds_arrput
@@ -407,6 +408,7 @@ CREDITS
 #define arraddnptr stbds_arraddnptr
 #define arraddnindex stbds_arraddnindex
 #define arrsetlen stbds_arrsetlen
+#define arrinitlen stbds_arrinitlen
 #define arrlast stbds_arrlast
 #define arrins stbds_arrins
 #define arrinsn stbds_arrinsn
@@ -539,8 +541,10 @@ extern "C"
 
 #define stbds_arrsetcap(a, n) (stbds_arrgrow(a, 0, n))
 #define stbds_arrsetlen(a, n) ((stbds_arrcap(a) < (size_t)(n) ? stbds_arrsetcap((a), (size_t)(n)), 0 : 0), (a) ? stbds_header(a)->length = (size_t)(n) : 0)
+#define stbds_arrinitlen(a, n) ((stbds_arrsetcap((a), (size_t)(n)), stbds_header(a)->length = (size_t)(n)))
 #define stbds_arrcap(a) ((a) ? stbds_header(a)->capacity : 0)
 #define stbds_arrlen(a) ((a) ? (ptrdiff_t)stbds_header(a)->length : 0)
+#define stbds_arrlennonull(a) ((ptrdiff_t)stbds_header(a)->length)
 #define stbds_arrlenu(a) ((a) ? stbds_header(a)->length : 0)
 #define stbds_arrput(a, v) (stbds_arrmaybegrow(a, 1), (a)[stbds_header(a)->length++] = (v))
 #define stbds_arrpush stbds_arrput	// synonym
