[stb_preprocess]
- add handling of #pragma directives in macro definitions
- add handling of #pragma message to create diagnostic outputs
- add and update comments in preprocessor.h
- unifying memory allocation mechanisms between stb_ds, stb_alloc and preprocessor - allow configuring all functions performing memory allocation/freeing via macros which should be specified via a header whose include path is specified via the STB_CONFIG macro (this will be included in stb_common.h if defined)
- don't pass an stb_arena to the include resolution callback; this is unnecessary complexity (implementation can handle its own memory allocation as long as the resulting strings have the appropriate lifetime) and indeed the arena used was global (and so not threadsafe)
- expose function allowing external modification of error modes so we can disable the error case for preprocessor directives not being at the start of a line (and so handle the case where we #define a macro whose value is a #pragma, or any other string possibly containing a # that we want to pass unmodified to the output i.e. the case of #define COMPILER_DEFINE #define which we use to keep #defines in the preprocessed code for handling by the platform compilers)
- fix leak of pp diagnostic messages
- remove internal caches for include resolution; they are problematic since (a) they aren't threadsafe, and (b) the caching of resolved filenames based on the include path found in source is inherently problematic (can potentially cause issues with same-named files in different folders included via relative paths)
- strip all code related to explicit file loading and lists of include paths; this has no purpose for our use case. note this makes the callbacks for file load/unload and include resolve mandatory, comments in preprocessor.h updated to reflect this.
- fix a few other memory leaks

diff --git a/Engine/Source/Developer/ShaderPreprocessor/Private/stb_preprocess/preprocessor.c b/Engine/Source/Developer/ShaderPreprocessor/Private/stb_preprocess/preprocessor.c
--- a/Engine/Source/Developer/ShaderPreprocessor/Private/stb_preprocess/preprocessor.c
+++ b/Engine/Source/Developer/ShaderPreprocessor/Private/stb_preprocess/preprocessor.c
@@ -106,6 +106,7 @@ static struct macro_definition predefined_defined = {"defined", 7, {0}, 0, MACRO
 #define strdup(x) _strdup(x)
 #endif
 
+
 typedef struct
 {
 	char* key;
@@ -160,7 +161,7 @@ static resolveinclude_callback_func resolveinclude_callback;
 
 typedef struct
 {
-	table_entry* table;	   // malloc
+	table_entry* table;	   // stb_pp_malloc
 	hashtable_pair* pair;  // stb_ds array
 	int first_unused_pair;
 	uint32 table_size;
@@ -174,10 +175,6 @@ typedef struct
 
 typedef struct pp_context
 {
-	char** include_paths;
-	int num_include_paths;
-	char** sys_include_paths;
-	int num_sys_include_paths;
 	pphash macro_map;				  // map of macro names to their definitions
 	macro_hash_entry* undef_map;	  // map of macro names to a flag for whether they've ever been undef'd, for more useful error messages
 	uint8 (*macro_name_filter)[256];  // 64KB table indicating whether 1 & 2-character identifiers are macros
@@ -416,7 +413,7 @@ static void output_line_directive(parse_state* cs)
 	}
 
 	// write the line directive
-	stb_assume(q != NULL);
+	STB_ASSUME(q != NULL);
 	
 	// ensure line directives start on a new line
 	if (q[-1] != '\n')
@@ -473,7 +470,7 @@ void pp_set_did_you_mean_threshold(int mode)
 void fill_where(parse_state* stack, pp_diagnostic* diagnostic, int* line_number)
 {
 	pp_where where = { 0 };
-	where.filename = strdup(stack->filename);
+	where.filename = STB_COMMON_STRDUP(stack->filename);
 	where.line_number = (*line_number > 0 ? *line_number : stack->src_line_number);
 	where.column = 0;
 	*line_number = 0;
@@ -482,7 +479,7 @@ void fill_where(parse_state* stack, pp_diagnostic* diagnostic, int* line_number)
 
 static int error_explicit_v(parse_state* ps, int code, int line_number, char* text, va_list va)
 {
-	stb_assume(ps != NULL);
+	STB_ASSUME(ps != NULL);
 	parse_state* stack;
 	pp_diagnostic d;
 	size_t len;
@@ -507,7 +504,7 @@ static int error_explicit_v(parse_state* ps, int code, int line_number, char* te
 	if (size < 1)
 		return 1;
 
-	d.message = (char*)malloc(size);
+	d.message = (char*)STB_COMMON_MALLOC(size);
 
 	// sanity check added to satisfy clang static analyzer
 	if (d.message == NULL)
@@ -692,9 +689,9 @@ void stringhash_init_table(pphash* ht, unsigned int size)
 {
 	table_entry t = {0}, *table;
 	unsigned int i;
-	stb_assume(size > 0u);
-	table = (table_entry*)malloc(sizeof(table[0]) * size);
-	stb_assume(table != NULL);
+	STB_ASSUME(size > 0u);
+	table = (table_entry*)STB_COMMON_MALLOC(sizeof(table[0]) * size);
+	STB_ASSUME(table != NULL);
 	ht->table_size = size;
 	ht->table_mask = size - 1;
 	ht->table = table;
@@ -906,7 +903,7 @@ void stringhash_put(pphash* ht, const char* key, size_t keylen, void* value)
 		for (i = 0; i < size; ++i)
 			if (table[i].hash > HASH_TOMBSTONE)
 				stringhash_insert_ref(ht, table[i]);
-		free(table);
+		STB_COMMON_FREE(table);
 	}
 }
 
@@ -923,7 +920,7 @@ void stringhash_destroy(pphash* ht)
 {
 	stb_arena_free(&ht->arena);
 	arrfree(ht->pair);
-	free(ht->table);
+	STB_COMMON_FREE(ht->table);
 	if (ht->verify)
 		shfree(ht->verify);
 }
@@ -1341,11 +1338,11 @@ static int copy_to_action_point(parse_state* cs)
 				uint8 ch = (uint8)*p++;
 				uint8 ch_class;
 				assert(cs->fast_dest != NULL || q < out + arrcap(out));
-				stb_assume(q != NULL);
+				STB_ASSUME(q != NULL);
 				*q++ = ch;
 				ch_class = pp_char_class[ch];
 				prev_state = state;
-				stb_assume(state < PP_STATE_active_count && ch_class < PP_CHAR_CLASS_count);
+				STB_ASSUME(state < PP_STATE_active_count && ch_class < PP_CHAR_CLASS_count);
 				state = pp_transition_table[state][ch_class];
 				line_count += pp_state_is_end_of_line[state];
 			} while (state < state_limit);
@@ -1459,7 +1456,7 @@ static int copy_to_action_point_macro_expansion(parse_state* cs)
 	int state;
 	const char* p = cs->src + cs->src_offset;
 	char* out = cs->dest;
-	stb_assume(out != NULL);
+	STB_ASSUME(out != NULL);
 	char* q;
 
 	if (cs->fast_dest)
@@ -1471,7 +1468,7 @@ static int copy_to_action_point_macro_expansion(parse_state* cs)
 		q = out + arrlen(out);
 	}
 
-	stb_assume(q != NULL);
+	STB_ASSUME(q != NULL);
 
 	state = PP_STATE_ready;
 
@@ -1881,6 +1878,49 @@ typedef struct
 static void process_directive(parse_state* cs, conditional_state* cons);
 static void maybe_expand_macro(parse_state* cs, struct macro_definition* pending, int in_macro_expansion, char* in_macro, struct macro_definition** md);
 
+enum
+{
+	// perfect hash function: sum of ASCII values mod 32
+	HASH_elif = 0,
+	HASH_include = 4,
+	HASH_endif = 6,
+	HASH_line = 8,
+	HASH_else = 9,
+	HASH_error = 10,
+	HASH_define = 11,
+	HASH_if = 15,
+	HASH_ifndef = 12,
+	HASH_undef = 18,
+	HASH_pragma = 24,
+	HASH_ifdef = 30,
+
+	HASH_none = 31,
+};
+
+// normally this comes from a safe buffer, but during disabled-by-#if scanning it's an uncopied directive
+static int parse_directive_after_hash(const char** ptr_p, int* newlines)
+{
+	int sum = 0;
+	const char* p = *ptr_p;
+	const char* start;
+	p = preprocessor_skip_whitespace_in_directive(p, newlines);
+	start = p;
+	while (char_is_pp_identifier_first(*p))
+		sum += (uint8)*p++;
+	sum &= (DIRECTIVE_HASH_SIZE - 1);
+	// @OPTIMIZE: use a masked 8-byte comparison?
+	if (0 == memcmp(start, directive_hash[sum].name, directive_hash[sum].name_len) && !char_is_pp_identifier(directive_hash[sum].name_len))
+	{
+		*ptr_p = p;
+		return sum;
+	}
+	else
+	{
+		*ptr_p = start;
+		return HASH_none;
+	}
+}
+
 static void preprocess_string(parse_state* cs, int in_macro_expansion, char* in_macro, struct macro_definition** final_function_macro)
 {
 	pp_context* c = cs->context;
@@ -2053,22 +2093,27 @@ static void preprocess_string(parse_state* cs, int in_macro_expansion, char* in_
 
 			case PP_STATE_saw_hash:
 			{
-				// special case for HLSL: need to handle the 1.#INF infinity constant and not error here.
-				int hlsl_inf = ((cs->src_offset >= 2) &&
-					((cs->src_offset + 3) < cs->src_length) &&
-					cs->src[cs->src_offset - 2] == '1' &&
-					cs->src[cs->src_offset - 1] == '.' &&
-					cs->src[cs->src_offset + 1] == 'I' &&
-					cs->src[cs->src_offset + 2] == 'N' &&
-					cs->src[cs->src_offset + 3] == 'F');
-
-				if (!hlsl_inf && do_error_code(cs, PP_RESULT_directive_not_at_start_of_line, "Preprocessor directive must appear at beginning of line."))
+				if (do_error_code(cs, PP_RESULT_directive_not_at_start_of_line, "Preprocessor directive must appear at beginning of line."))
+				{
 					return;
+				}
 				else
 				{
-					// if we're supposed to keep going, write it to the output; maybe you're using this with a language that allows # as a token
-					arrput(cs->dest, '#');
-					++cs->src_offset;
+					int lines, hash;
+					// scan and hash the directive name (p initialized to the character after the #)
+					const char* p = cs->src + 1;
+					hash = parse_directive_after_hash(&p, &lines);
+					if ((hash & (DIRECTIVE_HASH_SIZE - 1)) == HASH_pragma)
+					{
+						// allow (and process) #pragmas in #defines
+						process_directive(cs, &cons);
+					}
+					else
+					{
+						// copy all other directives (or other usage of #, i.e. HLSL infinity constant 1.#INF) through to the output
+						arrput(cs->dest, '#');
+						++cs->src_offset;
+					}
 				}
 				break;
 			}
@@ -2141,6 +2186,7 @@ static struct macro_definition* create_macro_definition(parse_state* ps, const c
 	m->symbol_name = NULL;
 	m->disabled = 0;
 	m->predefined = 0;
+	m->expansion = NULL;
 
 	// parse identifier
 	p = start;
@@ -2561,7 +2607,7 @@ static const char* copy_argument(const char* text, int* line_number, char** p_ou
 		}
 		oldlen = arrlen(out);
 		addlen = p - q;
-		stb_assume(out != NULL && oldlen > 0);
+		STB_ASSUME(out != NULL && oldlen > 0);
 		arrsetlen(out, oldlen + addlen);
 		memcpy(out + oldlen, q, addlen);
 
@@ -3271,6 +3317,7 @@ static int evaluate_if(parse_state* cs, const char* p, int* syntax_error)
 	{
 		*syntax_error = 0;
 	}
+	arrfree(ncs.dest);
 	return result;
 }
 
@@ -3300,48 +3347,6 @@ static const char* parse_pp_identifier_in_directive(parse_state* ps, const char*
 	return p;
 }
 
-enum
-{
-	// perfect hash function: sum of ASCII values mod 32
-	HASH_elif = 0,
-	HASH_include = 4,
-	HASH_endif = 6,
-	HASH_line = 8,
-	HASH_else = 9,
-	HASH_error = 10,
-	HASH_define = 11,
-	HASH_if = 15,
-	HASH_ifndef = 12,
-	HASH_undef = 18,
-	HASH_pragma = 24,
-	HASH_ifdef = 30,
-
-	HASH_none = 31,
-};
-
-// normally this comes from a safe buffer, but during disabled-by-#if scanning it's an uncopied directive
-static int parse_directive_after_hash(const char** ptr_p, int* newlines)
-{
-	int sum = 0;
-	const char* p = *ptr_p;
-	const char* start;
-	p = preprocessor_skip_whitespace_in_directive(p, newlines);
-	start = p;
-	while (char_is_pp_identifier_first(*p))
-		sum += (uint8)*p++;
-	sum &= (DIRECTIVE_HASH_SIZE - 1);
-	// @OPTIMIZE: use a masked 8-byte comparison?
-	if (0 == memcmp(start, directive_hash[sum].name, directive_hash[sum].name_len) && !char_is_pp_identifier(directive_hash[sum].name_len))
-	{
-		*ptr_p = p;
-		return sum;
-	}
-	else
-	{
-		*ptr_p = start;
-		return HASH_none;
-	}
-}
 
 enum
 {
@@ -3351,15 +3356,6 @@ enum
 };
 
 static char* preprocess_string_from_file(parse_state* ps, const char* filename, const char* text, size_t textlen, int conditional_compilation_nesting, int is_header);
-static const char* resolve_file(char path[MAX_FILENAME],
-						  parse_state* ps,
-						  ppbool check_source_dirs,
-						  ppbool check_normal_includes,
-						  ppbool check_system_includes,
-						  const char* filename,
-						  ptrdiff_t filename_len);
-static const char* load_file_with_nul_terminator(const char* filename, size_t* ptr_len, int* mode, void* custom_context);
-static void load_file_free(const char* filename, const char* filedata, int mode, void* custom_context);
 
 static void pragma_once(pp_context* c, const char* filename, char* macroname)
 {
@@ -3405,15 +3401,7 @@ static int process_include(parse_state* cs, const char* start, conditional_state
 			++p;
 		}
 
-		if (*p == '"')
-		{
-			filename = resolve_file(path, cs, 1, 1, 0, start_char, p - start_char);
-		}
-		else
-		{
-			filename = resolve_file(path, cs, 0, 0, 1, start_char, p - start_char);
-		}
-
+		filename = (*resolveinclude_callback)(start_char, p - start_char, cs->filename, c->custom_context);
 		if (filename == NULL)
 		{
 			memcpy(path, start_char, p - start_char);
@@ -3447,11 +3435,10 @@ static int process_include(parse_state* cs, const char* start, conditional_state
 			// if macro is defined, it's an include guard, which means if it's defined we can skip processing
 			if (md == NULL)
 			{
-				int mode;
 				const char* data;
 				char *result;
 				size_t len;
-				data = load_file_with_nul_terminator(filename, &len, &mode, c->custom_context);
+				data = (*loadfile_callback)(filename, c->custom_context, &len);
 
 				if (data == NULL)
 				{
@@ -3461,7 +3448,9 @@ static int process_include(parse_state* cs, const char* start, conditional_state
 
 				c->num_includes += 1;
 				result = preprocess_string_from_file(cs, filename, data, len, cons->conditional_compilation_nesting_level, 1);
-				load_file_free(filename, data, mode, c->custom_context);
+				
+				(*freefile_callback)(filename, data, c->custom_context);
+
 				if (cs->fast_dest)
 					cs->fast_dest = result;
 				else
@@ -3478,6 +3467,14 @@ static int process_include(parse_state* cs, const char* start, conditional_state
 	return ret;
 }
 
+static void macro_free(struct macro_definition* md)
+{
+	if (md->num_parameters >= 0)
+	{
+		arrfree(md->expansion);
+	}
+}
+
 // process # directives
 static void process_directive(parse_state* cs, conditional_state* cons)
 {
@@ -3635,6 +3632,47 @@ static void process_directive(parse_state* cs, conditional_state* cons)
 						}
 					}
 					break;
+				case 'm':
+					if (0 == strcmp(identifier, "message"))
+					{
+						const char* message_end;
+						p = scan_pp_identifier(p);
+						if (*p++ != '(')
+						{
+							do_error_code(cs, PP_RESULT_malformed_pragma_message, "Malformed #pragma message");
+							goto handled;
+						}
+						char quote = *p++;
+
+						if (!char_is_quote(quote))
+						{
+							do_error_code(cs, PP_RESULT_malformed_pragma_message, "Malformed #pragma message");
+							goto handled;
+						}
+
+						message_end = p;
+						while (*message_end != quote && !(char_is_end_of_line(*message_end)))
+							message_end++;
+
+						if (!char_is_quote(*message_end))
+						{
+							do_error_code(cs, PP_RESULT_malformed_pragma_message, "Malformed #pragma message");
+							goto handled;
+						}
+
+						ptrdiff_t len = (message_end - p) + 1;
+
+						pp_diagnostic diag;
+						diag.error_level = PP_RESULT_MODE_no_warning;
+						diag.diagnostic_code = PP_RESULT_ok;
+						diag.message = (char*)STB_COMMON_MALLOC(len);
+						memcpy(diag.message, p, len - 1);
+						diag.message[len - 1] = 0;
+						diag.where = 0;
+
+						arrpush(cs->context->error, diag);
+						goto handled;
+					}
 				}
 			}
 
@@ -3679,6 +3717,7 @@ static void process_directive(parse_state* cs, conditional_state* cons)
 				if (old != NULL)
 				{
 					// @TODO: compare new definition to old and produce error/warning
+					macro_free(old);
 					stringhash_delete(&c->macro_map, m->symbol_name, m->symbol_name_length);
 				}
 				stringhash_put(&c->macro_map, m->symbol_name, m->symbol_name_length, m);
@@ -3714,6 +3753,7 @@ static void process_directive(parse_state* cs, conditional_state* cons)
 				}
 				else
 				{
+					macro_free(md);
 					stringhash_delete(&c->macro_map, identifier, idlen);
 					shput(c->undef_map, identifier, (struct macro_definition*)(size_t)1);
 				}
@@ -4505,10 +4545,6 @@ static void define_macro(pphash* map, struct macro_definition* m)
 char* preprocess_file(char* output_autobuffer,
 	const char* filename,
 	void* custom_context,
-	char** include_paths,
-	int num_include_paths,
-	char** sys_include_paths,
-	int num_sys_include_paths,
 	struct macro_definition** predefined_macros,
 	int num_predefined_macros,
 	pp_diagnostic** pd,
@@ -4516,11 +4552,11 @@ char* preprocess_file(char* output_autobuffer,
 {
 	char* output = 0;
 	pp_context c = { 0 };
-	int main_mode, i;
+	int i;
 	size_t length;
 	const char* main_file;
 
-	main_file = load_file_with_nul_terminator(filename, &length, &main_mode, custom_context);
+	main_file = (*loadfile_callback)(filename, custom_context, &length);
 	*num_pd = 0;
 
 	if (main_file == 0)
@@ -4535,10 +4571,6 @@ char* preprocess_file(char* output_autobuffer,
 		return NULL;
 	}
 
-	c.include_paths = include_paths;
-	c.sys_include_paths = sys_include_paths;
-	c.num_include_paths = num_include_paths;
-	c.num_sys_include_paths = num_sys_include_paths;
 	c.custom_context = custom_context;
 
 	// allocate the macro definition dictionary very large so (a) it's unlikely to grow,
@@ -4546,7 +4578,7 @@ char* preprocess_file(char* output_autobuffer,
 	// backfire by slowing down the initialization
 	stringhash_create(&c.macro_map, 4096);
 
-	c.macro_name_filter = malloc(65536);
+	c.macro_name_filter = STB_COMMON_MALLOC(65536);
 	if (c.macro_name_filter == NULL)
 		return NULL;
 
@@ -4582,11 +4614,25 @@ char* preprocess_file(char* output_autobuffer,
 		}
 	}
 
-	load_file_free(filename, main_file, main_mode, custom_context);
+	(*freefile_callback)(filename, main_file, custom_context);
+
+	for (i = 0; i < arrlen(c.macro_map.pair); ++i)
+	{
+		if (c.macro_map.pair[i].key)
+		{
+			struct macro_definition* md = (struct macro_definition*)c.macro_map.pair[i].value;
+			if (!md->predefined)
+			{
+				macro_free(md);
+			}
+		}
+	}
 
 	shfree(c.once_map);
+	shfree(c.undef_map);
+	arrfree(c.ifdef_stack);
 	stringhash_destroy(&c.macro_map);
-	free(c.macro_name_filter);
+	STB_COMMON_FREE(c.macro_name_filter);
 	stb_arena_free(&c.macro_arena);
 
 	if (c.error)
@@ -4620,9 +4666,10 @@ void preprocessor_file_free(char* text, pp_diagnostic* pd)
 // in which an uninitialized pp_where can be added to the array, so this is bogus
 #pragma warning(disable:6001)
 		for (j = 0; j < arrlen(pd[i].where); ++j)
-			free(pd[i].where[j].filename);
+			STB_COMMON_FREE(pd[i].where[j].filename);
 #pragma warning(pop)
 		arrfree(pd[i].where);
+		STB_COMMON_FREE(pd[i].message);
 	}
 	arrfree(pd);
 }
@@ -4685,6 +4732,8 @@ static void init_char_type(void)
 	pp_char_flags2['/'] |= CHAR2_IS_SLASH;
 }
 
+
+
 void init_preprocessor(
 	loadfile_callback_func load_callback, 
 	freefile_callback_func free_callback,
@@ -4716,231 +4765,4 @@ void init_preprocessor(
 	resolveinclude_callback = resolve_callback;
 }
 
-// File reader
-
-typedef struct
-{
-	char* key;	// resolved path
-	char* contents;
-	int contents_in_use;  // when multiprocessing or recursing, don't discard contents
-	size_t contents_size;
-	int contents_allocation_mode;
-} cached_header;
-
-typedef struct
-{
-	const char* key;
-	const char* value;
-} resolved_header;
-
-static cached_header* cached_header_map;
-static size_t cached_header_total_size;
-static resolved_header* resolved_callback_header_map;
-static resolved_header* resolved_header_map;
-static resolved_header* resolved_system_header_map;
-static struct stb_arena resolve_arena;
-
-static size_t max_cached_headers = sizeof(size_t) > 4 ? (128U << 20)  // 128MB in 64-bit
-													  : (32U << 20);  //  32MB in 32-bit
-
-void preprocessor_header_cache_set_size(size_t size)
-{
-	max_cached_headers = size;
-}
-
-void preprocessor_header_cache_flush(void)
-{
-	int i;
-	for (i = (int)hmlen(cached_header_map); i >= 0;)
-	{
-		if (cached_header_map[i].contents_in_use)
-			--i;  // can't delete it
-		else
-		{
-			free(cached_header_map[i].contents);
-			cached_header_map[i].contents = 0;
-		}
-	}
-}
-
-void init_preprocessor_cache(void)
-{
-	sh_new_arena(resolved_header_map);
-	sh_new_arena(resolved_system_header_map);
-}
-
-static int find_file_path(char path[MAX_FILENAME], char* filename, char** path_list, ptrdiff_t path_list_len)
-{
-	int i;
-	FILE* f = NULL;
-
-	for (i = 0; i < path_list_len; ++i)
-	{
-		sprintf(path, "%s/%s", path_list[i], filename);
-		f = fopen(path, "rb");
-		if (f)
-		{
-			fclose(f);
-			return 1;
-		}
-	}
-
-	return 0;
-}
-
-static char* my_strcpy(char* dest, const char* src)
-{
-	while (*src)
-		*dest++ = *src++;
-	*dest = 0;
-	return dest;
-}
-
-static const char* resolve_file(char path[MAX_FILENAME],
-						  parse_state* ps,
-						  ppbool check_source_dirs,
-						  ppbool check_normal_includes,
-						  ppbool check_system_includes,
-						  const char* filename,
-						  ptrdiff_t filename_len)
-{
-	pp_context* c = ps->context;
-
-	char filenamez[MAX_FILENAME];
-	strncpy(filenamez, filename, filename_len);
-	filenamez[filename_len] = 0;
-
-	if (resolveinclude_callback != NULL)
-	{
-		const char* resolved = shget(resolved_callback_header_map, path);
-		if (resolved)
-			return resolved;
-
-		resolved  = (*resolveinclude_callback)(filename, filename_len, ps->filename, &resolve_arena, c->custom_context);
-		if (resolved)
-		{
-			shput(resolved_callback_header_map, filenamez, resolved);
-			return resolved;
-		}
-	}
-
-	if (check_source_dirs)
-	{
-		do
-		{
-			FILE* f;
-			char* end;
-			if (ps->parent != NULL)
-			{
-				end = my_strcpy(path, ps->filename);
-				while (end > path && end[-1] != '/' && end[-1] != '\\')
-					--end;
-				strcpy(end, filenamez);
-				f = fopen(path, "rb");
-				if (f)
-				{
-					fclose(f);
-					return path;
-				}
-			}
-			ps = ps->parent;
-		} while (ps);
-	}
-
-	if (check_normal_includes)
-	{
-		const char* resolved = shget(resolved_header_map, filenamez);
-		if (resolved)
-			return resolved;
-
-		if (find_file_path(path, filenamez, c->include_paths, c->num_include_paths))
-		{
-			resolved = stb_arena_alloc_string(&resolve_arena, path);
-			shput(resolved_header_map, filenamez, resolved);
-			return resolved;
-		}
-	}
-
-	if (check_system_includes)
-	{
-		const char* resolved = shget(resolved_system_header_map, filenamez);
-		if (resolved)
-			return resolved;
-
-		if (find_file_path(path, filenamez, c->sys_include_paths, c->num_sys_include_paths))
-		{
-			resolved = stb_arena_alloc_string(&resolve_arena, path);
-			shput(resolved_system_header_map, filenamez, resolved);
-			return resolved;
-		}
-	}
-
-	return NULL;
-}
-
-#define LOAD_FILE_MODE_callback -3
-#define LOAD_FILE_MODE_malloced -2
-#define LOAD_FILE_MODE_memmappend -1
-#define LOAD_FILE_MODE_multi_memmaped 0	 // actual positive value is number of 4K pages past end of pointer that second mapping occurs
-
-static const char* load_file_with_nul_terminator(const char* filename, size_t* ptr_len, int* mode, void* custom_context)
-{
-	if (loadfile_callback != NULL)
-	{
-		*mode = LOAD_FILE_MODE_callback;
-		return (*loadfile_callback)(filename, custom_context, ptr_len);
-	}
-	else
-	{
-		char* p;
-		size_t len;
-		FILE* f = fopen(filename, "rb");
-		*mode = LOAD_FILE_MODE_malloced;
-		if (f == 0)
-			return NULL;
-		fseek(f, 0, SEEK_END);
-		len = ftell(f);
-		fseek(f, 0, SEEK_SET);
-		p = (char*)malloc(len + 1);
-		if (p == NULL)
-			return NULL;
-
-		p[len] = 0;
-		fread(p, 1, len, f);
-		fclose(f);
-		if (ptr_len)
-			*ptr_len = len + 1;
-		return p;
-	}
-}
-
-static void load_file_free(const char* filename, const char* filedata, int mode, void* custom_context)
-{
-	if (freefile_callback != NULL && mode == LOAD_FILE_MODE_callback)
-	{
-		(*freefile_callback)(filename, filedata, custom_context);
-	}
-	else if (mode == LOAD_FILE_MODE_malloced)
-	{
-		free((void*)filedata);
-	}
-}
-
-#if 0
-#include <windows.h>
-void *stb_mmap(char *filename, int *n)
-{
-   HANDLE f = CreateFile(filename, GENERIC_READ, FILE_SHARE_READ,  NULL, OPEN_EXISTING, FILE_ATTRIBUTE_NORMAL, NULL);
-   HANDLE m;
-   void *p;
-   if (!f) stb_fatal("Couldn't CreateFile() on '%s'", filename);
-   if (n) *n = GetFileSize(f, NULL);
-   m = CreateFileMapping(f, NULL, PAGE_READONLY, 0,0, NULL);
-   if (!m) stb_fatal("Couldn't CreateFileMapping() on '%s'", filename);
-   p = MapViewOfFile(m, FILE_MAP_READ, 0,0,0);
-   if (!p) stb_fatal("Couldn't MapViewOfFile() on '%s'", filename);
-   return p;
-}
-#endif
-
 #pragma warning(pop)
diff --git a/Engine/Source/Developer/ShaderPreprocessor/Private/stb_preprocess/preprocessor.h b/Engine/Source/Developer/ShaderPreprocessor/Private/stb_preprocess/preprocessor.h
--- a/Engine/Source/Developer/ShaderPreprocessor/Private/stb_preprocess/preprocessor.h
+++ b/Engine/Source/Developer/ShaderPreprocessor/Private/stb_preprocess/preprocessor.h
@@ -22,6 +22,7 @@ enum
 	PP_RESULT_unexpected_character_after_end_of_directive,
 	PP_RESULT_undef_of_predefined_macro,
 	PP_RESULT_identifier_too_long,
+	PP_RESULT_malformed_pragma_message,
 
 	PP_RESULT_count
 };
@@ -36,6 +37,7 @@ enum
 	PP_RESULT_MODE_no_warning
 };
 
+
 struct macro_definition;
 struct stb_arena;
 
@@ -61,39 +63,46 @@ typedef struct
 #define STB_PP_DEF extern
 #endif
 
-// callback function can be optionally specified to init_preprocessor to override file loading functionality
-// it's probably necessary to provide the freefile callback as well if you provide this (unless you are ok with
-// the default behaviour of calling "free" on the returned pointer)
+// callback function required to be passed to init_preprocessor for implementation of file loading functionality
 typedef const char* (*loadfile_callback_func)(const char* filename, void* custom_context, size_t* out_length);
 
-// callback function can be optionally specified to init_preprocessor to override freeing loaded files
-// it's probably necessary to provide the loadfile callback as well if you provide this
+// callback function required to be passed to init_preprocessor for implementation of freeing loaded files
 typedef void (*freefile_callback_func)(const char* filename, const char* loaded_file, void* custom_context);
 
-// callback function can be optionally specified to init_preprocessor to override include resolution; in the
-// current incarnation this will still attempt file-based loading using the default behaviour in the event
-// that this callback returns null for a given file path.
-// note: returned resolved paths are expected to be allocated for the lifetime of preprocessor execution; 
-// for simplicity use stb_arena_alloc_string using the given "path_arena" stb_arena
-typedef const char* (*resolveinclude_callback_func)(const char* path, unsigned int path_len, const char* parent, struct stb_arena* path_arena, void* custom_context);
+// callback function required to be passed to init_preprocessor for implementation of include resolution
+// the callback is expected to handle both absolute and relative paths (the latter relative to the given
+// "parent" file path)
+// note: returned resolved paths must remain valid for the lifetime of a single preprocessor execution
+typedef const char* (*resolveinclude_callback_func)(const char* path, unsigned int path_len, const char* parent, void* custom_context);
 
-STB_PP_DEF void preprocessor_test(void);
+// init function must be called once before any invocations of preprocess_file
 STB_PP_DEF void init_preprocessor(
 	loadfile_callback_func load_callback,
 	freefile_callback_func free_callback,
 	resolveinclude_callback_func resolve_callback);
+
+// main preprocessing execution function; returns the preprocessed string and sets any diagnostic messages
+// in the pd array (setting num_pd to the diagnostic count). preprocessor_file_free should be called for 
+// each invocation to free any allocated memory.
 STB_PP_DEF char* preprocess_file(
 	char* output_storage,
 	const char* filename,
 	void* custom_context,
-	char** include_paths,
-	int num_include_paths,
-	char** sys_include_paths,
-	int num_sys_include_paths,
 	struct macro_definition** predefined_macros,
 	int num_predefined_macros,
 	pp_diagnostic** pd,
 	int* num_pd);
+
+// frees memory allocated by preprocess_file (preprocessed results and diagnostic messages)
 STB_PP_DEF void preprocessor_file_free(char* text, pp_diagnostic* pd);
-STB_PP_DEF struct macro_definition* pp_define(struct stb_arena* a, const char* p);
-STB_PP_DEF unsigned int preprocessor_hash(const char* data, size_t len);	// hash function for string of length 'len', roughly xxhash32
+
+// constructs a preprocessor definition, allocating on the given arena. "def" should be in the format:
+// DEFINE_NAME DEFINE_VALUE
+STB_PP_DEF struct macro_definition* pp_define(struct stb_arena* a, const char* def);
+
+// override behaviour of a particular error result (can be used to, for instance, downgrade
+// an error to a warning or a "no warning" i.e. valid case).
+STB_PP_DEF void pp_set_warning_mode(int result_code, int result_mode);
+
+// hash function for string of length 'len', roughly xxhash32
+STB_PP_DEF unsigned int preprocessor_hash(const char* data, size_t len);
diff --git a/Engine/Source/Developer/ShaderPreprocessor/Private/stb_preprocess/stb_alloc.h b/Engine/Source/Developer/ShaderPreprocessor/Private/stb_preprocess/stb_alloc.h
--- a/Engine/Source/Developer/ShaderPreprocessor/Private/stb_preprocess/stb_alloc.h
+++ b/Engine/Source/Developer/ShaderPreprocessor/Private/stb_preprocess/stb_alloc.h
@@ -68,7 +68,7 @@ void* stb_arena_alloc_aligned(struct stb_arena* a, size_t size, size_t align)
 		// grow the list of blocks if necessary
 		if (a->num_blocks == 0)
 		{
-			a->blocks = (void**)malloc(sizeof(a->blocks[0]) * 32);
+			a->blocks = (void**)STB_COMMON_MALLOC(sizeof(a->blocks[0]) * 32);
 		}
 		else if (a->num_blocks >= 32)
 		{
@@ -76,15 +76,15 @@ void* stb_arena_alloc_aligned(struct stb_arena* a, size_t size, size_t align)
 			if ((a->num_blocks & (a->num_blocks - 1)) == 0)
 			{
 				// grow table
-				void* ptr = realloc(a->blocks, sizeof(a->blocks[0]) * a->num_blocks * 2);
+				void* ptr = STB_COMMON_REALLOC(a->blocks, sizeof(a->blocks[0]) * a->num_blocks * 2);
 				if (ptr)
 					a->blocks = (void**)ptr;
 			}
 		}
 
 		// allocate and record allocation
-		a->last_block = (unsigned char*)malloc(a->last_block_size + 31);
-		stb_assume(a->blocks != NULL);
+		a->last_block = (unsigned char*)STB_COMMON_MALLOC(a->last_block_size + 31);
+		STB_ASSUME(a->blocks != NULL);
 		a->blocks[a->num_blocks++] = a->last_block;
 		a->last_block_alloc_offset = 0;
 
@@ -144,8 +144,8 @@ void stb_arena_free(struct stb_arena* a)
 {
 	int i;
 	for (i = 0; i < a->num_blocks; ++i)
-		free(a->blocks[i]);
-	free(a->blocks);
+		STB_COMMON_FREE(a->blocks[i]);
+	STB_COMMON_FREE(a->blocks);
 	a->blocks = 0;
 	a->num_blocks = 0;
 	a->last_block_size = 0;
diff --git a/Engine/Source/Developer/ShaderPreprocessor/Private/stb_preprocess/stb_common.h b/Engine/Source/Developer/ShaderPreprocessor/Private/stb_preprocess/stb_common.h
--- a/Engine/Source/Developer/ShaderPreprocessor/Private/stb_preprocess/stb_common.h
+++ b/Engine/Source/Developer/ShaderPreprocessor/Private/stb_preprocess/stb_common.h
@@ -2,8 +2,33 @@
 
 #pragma once
 
+#define STB_COMMON_XSTR(s) STB_COMMON_STR(s)
+#define STB_COMMON_STR(s) #s
+#ifdef STB_CONFIG
+#include STB_COMMON_XSTR(STB_CONFIG)
+#endif
+
+// The following macros can be defined to override memory allocation for preprocessor/stbds operation:
+// STB_COMMON_MALLOC
+// STB_COMMON_REALLOC
+// STB_COMMON_FREE
+// STB_COMMON_STRDUP
+// Note that overriding these is all-or-nothing and so will error if only some are defined.
+// By default (if not specified) malloc/realloc/free/strdup will be used.
+#if !defined(STB_COMMON_MALLOC) && !defined(STB_COMMON_REALLOC) && !defined(STB_COMMON_FREE) && !defined(STB_COMMON_STRDUP)
+#include <stdlib.h>
+#define STB_COMMON_MALLOC(s) malloc(s)
+#define STB_COMMON_REALLOC(p, s) realloc(p, s)
+#define STB_COMMON_FREE(p) free(p)
+#define STB_COMMON_STRDUP(s) strdup(s)
+#elif defined(STB_COMMON_MALLOC) && defined(STB_COMMON_REALLOC) && defined(STB_COMMON_FREE) && defined(STB_COMMON_STRDUP)
+// all are defined
+#else
+#error "You must define all or none of STB_COMMON_MALLOC/STB_COMMON_REALLOC/STB_COMMON_FREE/STB_COMMON_STRDUP
+#endif
+
 #if defined(_Analysis_assume_)
-#define stb_assume(expr) _Analysis_assume_(expr)
+#define STB_ASSUME(expr) _Analysis_assume_(expr)
 #else
-#define stb_assume(expr)
+#define STB_ASSUME(expr)
 #endif
\ No newline at end of file
diff --git a/Engine/Source/Developer/ShaderPreprocessor/Private/stb_preprocess/stb_ds.h b/Engine/Source/Developer/ShaderPreprocessor/Private/stb_preprocess/stb_ds.h
--- a/Engine/Source/Developer/ShaderPreprocessor/Private/stb_preprocess/stb_ds.h
+++ b/Engine/Source/Developer/ShaderPreprocessor/Private/stb_preprocess/stb_ds.h
@@ -43,16 +43,6 @@ COMPILE-TIME OPTIONS
 	 hash table insertion about 20% slower on 4- and 8-byte keys, 5% slower on
 	 64-byte keys, and 10% slower on 256-byte keys on my test computer.
 
-  #define STBDS_REALLOC(context,ptr,size) better_realloc
-  #define STBDS_FREE(context,ptr)         better_free
-
-	 These defines only need to be set in the file containing #define STB_DS_IMPLEMENTATION.
-
-	 By default stb_ds uses stdlib realloc() and free() for memory management. You can
-	 substitute your own functions instead by defining these symbols. You must either
-	 define both, or neither. Note that at the moment, 'context' will always be NULL.
-	 @TODO add an array/hash initialization function that takes a memory context pointer.
-
   #define STBDS_UNIT_TESTS
 
 	 Defines a function stbds_unit_tests() that checks the functioning of the data structures.
@@ -456,15 +446,6 @@ CREDITS
 #define strreset stbds_strreset
 #endif
 
-#if defined(STBDS_REALLOC) && !defined(STBDS_FREE) || !defined(STBDS_REALLOC) && defined(STBDS_FREE)
-#error "You must define both STBDS_REALLOC and STBDS_FREE, or neither."
-#endif
-#if !defined(STBDS_REALLOC) && !defined(STBDS_FREE)
-#include <stdlib.h>
-#define STBDS_REALLOC(c, p, s) realloc(p, s)
-#define STBDS_FREE(c, p) free(p)
-#endif
-
 #ifdef _MSC_VER
 #define STBDS_NOTUSED(v) (void)(v)
 #else
@@ -554,7 +535,7 @@ extern "C"
 #define stbds_arraddnoff(a, n) (stbds_arrmaybegrow(a, n), stbds_header(a)->length += (n), stbds_header(a)->length - (n))
 #define stbds_arraddnindex(a, n) stbds_arraddnoff(a, n)
 #define stbds_arrlast(a) ((a)[stbds_header(a)->length - 1])
-#define stbds_arrfree(a) ((void)((a) ? STBDS_FREE(NULL, stbds_header(a)) : (void)0), (a) = NULL)
+#define stbds_arrfree(a) ((void)((a) ? STB_COMMON_FREE(stbds_header(a)) : (void)0), (a) = NULL)
 #define stbds_arrdel(a, i) stbds_arrdeln(a, i, 1)
 #define stbds_arrdeln(a, i, n) (memmove(&(a)[i], &(a)[(i) + (n)], sizeof *(a) * (stbds_header(a)->length - (n) - (i))), stbds_header(a)->length -= (n))
 #define stbds_arrdelswap(a, i) ((a)[i] = stbds_arrlast(a), stbds_header(a)->length -= 1)
@@ -779,8 +760,8 @@ void* stbds_arrgrowf(void* a, size_t elemsize, size_t addlen, size_t min_cap)
 	// if (num_prev < 65536) if (a) prev_allocs[num_prev++] = (int *) ((char *) a+1);
 	// if (num_prev == 2201)
 	//   num_prev = num_prev;
-	b = STBDS_REALLOC(NULL, (a) ? stbds_header(a) : 0, elemsize * min_cap + sizeof(stbds_array_header));
-	stb_assume(b != NULL);
+	b = STB_COMMON_REALLOC((a) ? stbds_header(a) : 0, elemsize * min_cap + sizeof(stbds_array_header));
+	STB_ASSUME(b != NULL);
 
 	// if (num_prev < 65536) prev_allocs[num_prev++] = (int *) (char *) b;
 	b = (char*)b + sizeof(stbds_array_header);
@@ -882,9 +863,9 @@ static size_t stbds_log2(size_t slot_count)
 static stbds_hash_index* stbds_make_hash_index(size_t slot_count, stbds_hash_index* ot)
 {
 	stbds_hash_index* t;
-	t = (stbds_hash_index*)STBDS_REALLOC(NULL, 0,
-										 (slot_count >> STBDS_BUCKET_SHIFT) * sizeof(stbds_hash_bucket) + sizeof(stbds_hash_index) + STBDS_CACHE_LINE_SIZE - 1);
-	stb_assume(t != NULL);
+	t = (stbds_hash_index*)STB_COMMON_REALLOC(0,
+		(slot_count >> STBDS_BUCKET_SHIFT) * sizeof(stbds_hash_bucket) + sizeof(stbds_hash_index) + STBDS_CACHE_LINE_SIZE - 1);
+	STB_ASSUME(t != NULL);
 
 	t->storage = (stbds_hash_bucket*)STBDS_ALIGN_FWD((size_t)(t + 1), STBDS_CACHE_LINE_SIZE);
 	t->slot_count = slot_count;
@@ -1265,12 +1246,12 @@ void stbds_hmfree_func(void* a, size_t elemsize)
 			size_t i;
 			// skip 0th element, which is default
 			for (i = 1; i < stbds_header(a)->length; ++i)
-				STBDS_FREE(NULL, *(char**)((char*)a + elemsize * i));
+				STB_COMMON_FREE(*(char**)((char*)a + elemsize * i));
 		}
 		stbds_strreset(&stbds_hash_table(a)->string);
 	}
-	STBDS_FREE(NULL, stbds_header(a)->hash_table);
-	STBDS_FREE(NULL, stbds_header(a));
+	STB_COMMON_FREE(stbds_header(a)->hash_table);
+	STB_COMMON_FREE(stbds_header(a));
 }
 
 static ptrdiff_t stbds_hm_find_slot(void* a, size_t elemsize, void* key, size_t keysize, size_t keyoffset, int mode)
@@ -1429,7 +1410,7 @@ void* stbds_hmput_key(void* a, size_t elemsize, void* key, size_t keysize, int m
 		slot_count = (table == NULL) ? STBDS_BUCKET_LENGTH : table->slot_count * 2;
 		nt = stbds_make_hash_index(slot_count, table);
 		if (table)
-			STBDS_FREE(NULL, table);
+			STB_COMMON_FREE(table);
 		else
 			nt->string.mode = mode >= STBDS_HM_STRING ? STBDS_SH_DEFAULT : 0;
 		stbds_header(a)->hash_table = table = nt;
@@ -1599,7 +1580,7 @@ void* stbds_hmdel_key(void* a, size_t elemsize, void* key, size_t keysize, size_
 				b->index[i] = STBDS_INDEX_DELETED;
 
 				if (mode == STBDS_HM_STRING && table->string.mode == STBDS_SH_STRDUP)
-					STBDS_FREE(NULL, *(char**)((char*)a + elemsize * old_index));
+					STB_COMMON_FREE(*(char**)((char*)a + elemsize * old_index));
 
 				// if indices are the same, memcpy is a no-op, but back-pointer-fixup will fail, so skip
 				if (old_index != final_index)
@@ -1623,13 +1604,13 @@ void* stbds_hmdel_key(void* a, size_t elemsize, void* key, size_t keysize, size_
 				if (table->used_count < table->used_count_shrink_threshold && table->slot_count > STBDS_BUCKET_LENGTH)
 				{
 					stbds_header(raw_a)->hash_table = stbds_make_hash_index(table->slot_count >> 1, table);
-					STBDS_FREE(NULL, table);
+					STB_COMMON_FREE(table);
 					STBDS_STATS(++stbds_hash_shrink);
 				}
 				else if (table->tombstone_count > table->tombstone_count_threshold)
 				{
 					stbds_header(raw_a)->hash_table = stbds_make_hash_index(table->slot_count, table);
-					STBDS_FREE(NULL, table);
+					STB_COMMON_FREE(table);
 					STBDS_STATS(++stbds_hash_rebuild);
 				}
 
@@ -1645,7 +1626,7 @@ static char* stbds_strdup(char* str)
 	// to keep replaceable allocator simple, we don't want to use strdup.
 	// rolling our own also avoids problem of strdup vs _strdup
 	size_t len = strlen(str) + 1;
-	char* p = (char*)STBDS_REALLOC(NULL, 0, len);
+	char* p = (char*)STB_COMMON_REALLOC(0, len);
 	memmove(p, str, len);
 	return p;
 }
@@ -1680,7 +1661,7 @@ char* stbds_stralloc(stbds_string_arena* a, char* str)
 			// note that we still advance string_block so block size will continue
 			// increasing, so e.g. if somebody only calls this with 1000-long strings,
 			// eventually the arena will start doubling and handling those as well
-			stbds_string_block* sb = (stbds_string_block*)STBDS_REALLOC(NULL, 0, sizeof(*sb) + len);
+			stbds_string_block* sb = (stbds_string_block*)STB_COMMON_REALLOC(0, sizeof(*sb) + len);
 			if (!sb)
 				return NULL;
 			memmove(sb->storage, str, len);
@@ -1700,7 +1681,7 @@ char* stbds_stralloc(stbds_string_arena* a, char* str)
 		}
 		else
 		{
-			stbds_string_block* sb = (stbds_string_block*)STBDS_REALLOC(NULL, 0, sizeof(*sb) + blocksize);
+			stbds_string_block* sb = (stbds_string_block*)STB_COMMON_REALLOC(0, sizeof(*sb) + blocksize);
 			if (!sb)
 				return NULL;
 			sb->next = a->storage;
@@ -1723,7 +1704,7 @@ void stbds_strreset(stbds_string_arena* a)
 	while (x)
 	{
 		y = x->next;
-		STBDS_FREE(NULL, x);
+		STB_COMMON_FREE(x);
 		x = y;
 	}
 	memset(a, 0, sizeof(*a));
