[shaders] fix missing closing " in stb header error msg (this path should not be hit but it's causing problems for IncludeTool static analyzer)

diff --git a/Engine/Source/Developer/ShaderPreprocessor/Private/stb_preprocess/stb_common.h b/Engine/Source/Developer/ShaderPreprocessor/Private/stb_preprocess/stb_common.h
--- a/Engine/Source/Developer/ShaderPreprocessor/Private/stb_preprocess/stb_common.h
+++ b/Engine/Source/Developer/ShaderPreprocessor/Private/stb_preprocess/stb_common.h
@@ -24,7 +24,7 @@
 #elif defined(STB_COMMON_MALLOC) && defined(STB_COMMON_REALLOC) && defined(STB_COMMON_FREE) && defined(STB_COMMON_STRDUP)
 // all are defined
 #else
-#error "You must define all or none of STB_COMMON_MALLOC/STB_COMMON_REALLOC/STB_COMMON_FREE/STB_COMMON_STRDUP
+#error "You must define all or none of STB_COMMON_MALLOC/STB_COMMON_REALLOC/STB_COMMON_FREE/STB_COMMON_STRDUP"
 #endif
 
 #if defined(_Analysis_assume_)
