[shaders] add support to run STB preprocessor instead of MCPP; add a cvar to disable this and instead use the legacy preprocessor (currently this is set to true - i.e. still using MCPP - by default)

diff --git a/Engine/Shaders/Public/ShaderVersion.ush b/Engine/Shaders/Public/ShaderVersion.ush
--- a/Engine/Shaders/Public/ShaderVersion.ush
+++ b/Engine/Shaders/Public/ShaderVersion.ush
@@ -2,4 +2,4 @@
 // Each time the console command is executed it generates a new GUID. As this file is included
 // in Platform.ush (which should be included in any shader) it allows to invalidate the shader DDC.
 // 
-// GUID = 28A1A473-9B5D-40E9-B951-B09A08EB4D7F
\ No newline at end of file
+// GUID = BDF733E1-F4F8-4332-A403-7C2DAAC8A2E9
\ No newline at end of file
diff --git a/Engine/Source/Developer/ShaderPreprocessor/Private/PreprocessorPrivate.h b/Engine/Source/Developer/ShaderPreprocessor/Private/PreprocessorPrivate.h
--- a/Engine/Source/Developer/ShaderPreprocessor/Private/PreprocessorPrivate.h
+++ b/Engine/Source/Developer/ShaderPreprocessor/Private/PreprocessorPrivate.h
@@ -15,11 +15,11 @@ enum class EMessageType
 };
 
 /**
- * Filter MCPP errors.
+ * Filter preprocessor errors.
  * @param ErrorMsg - The error message.
  * @returns true if the message is valid and has not been filtered out.
  */
-inline EMessageType FilterMcppError(const FString& ErrorMsg)
+inline EMessageType FilterPreprocessorError(const FString& ErrorMsg)
 {
 	const TCHAR* SubstringsToFilter[] =
 	{
@@ -49,6 +49,10 @@ static void ExtractDirective(FString& OutString, FString WarningString)
 	static const FString PrefixString = TEXT("UESHADERMETADATA_");
 	uint32 DirectiveStartPosition = WarningString.Find(PrefixString) + PrefixString.Len();
 	uint32 DirectiveEndPosition = WarningString.Find(TEXT("\n"));
+	if (DirectiveEndPosition == INDEX_NONE)
+	{
+		DirectiveEndPosition = WarningString.Len();
+	}
 	OutString = WarningString.Mid(DirectiveStartPosition, (DirectiveEndPosition - DirectiveStartPosition));
 }
 
@@ -84,7 +88,7 @@ static bool ParseMcppErrors(TArray<FShaderCompilerError>& OutErrors, TArray<FStr
 					Message.TrimStartAndEndInline();
 
 					// Ignore the warning about files that don't end with a newline.
-					switch (FilterMcppError(Message))
+					switch (FilterPreprocessorError(Message))
 					{
 						case EMessageType::Error:
 						{
diff --git a/Engine/Source/Developer/ShaderPreprocessor/Private/ShaderPreprocessor.cpp b/Engine/Source/Developer/ShaderPreprocessor/Private/ShaderPreprocessor.cpp
--- a/Engine/Source/Developer/ShaderPreprocessor/Private/ShaderPreprocessor.cpp
+++ b/Engine/Source/Developer/ShaderPreprocessor/Private/ShaderPreprocessor.cpp
@@ -7,7 +7,26 @@
 #include "Modules/ModuleManager.h"
 #include "PreprocessorPrivate.h"
 
-IMPLEMENT_MODULE(FDefaultModuleImpl, ShaderPreprocessor);
+#define STB_DS_IMPLEMENTATION
+#define STB_ALLOC_IMPLEMENTATION
+#include "stb_preprocess/preprocessor.h"
+#include "stb_preprocess/stb_alloc.h"
+#include "stb_preprocess/stb_ds.h"
+
+namespace
+{
+	const FString PlatformHeader = TEXT("/Engine/Public/Platform.ush");
+	const FString PlatformHeaderLowerCase = PlatformHeader.ToLower();
+	void LogMandatoryHeaderError(const FShaderCompilerInput& Input, FShaderCompilerOutput& Output)
+	{
+		FShaderCompilerError Error;
+		Error.ErrorVirtualFilePath = Input.VirtualSourceFilePath;
+		Error.ErrorLineString = TEXT("1");
+		Error.StrippedErrorMessage = FString::Printf(TEXT("Error: Shader is required to include %s"), *PlatformHeader);
+
+		Output.Errors.Add(Error);
+	}
+}
 
 /**
  * Append defines to an MCPP command line.
@@ -54,7 +73,7 @@ public:
 
 	bool HasIncludedMandatoryHeaders() const
 	{
-		return CachedFileContents.Contains(TEXT("/Engine/Public/Platform.ush"));
+		return CachedFileContents.Contains(PlatformHeader);
 	}
 
 private:
@@ -211,35 +230,12 @@ static void DumpShaderDefinesAsCommentedCode(const FShaderCompilerInput& ShaderI
 
 //////////////////////////////////////////////////////////////////////////
 
-/**
- * Preprocess a shader.
- * @param OutPreprocessedShader - Upon return contains the preprocessed source code.
- * @param ShaderOutput - ShaderOutput to which errors can be added.
- * @param ShaderInput - The shader compiler input.
- * @param AdditionalDefines - Additional defines with which to preprocess the shader.
- * @param bShaderDumpDefinesAsCommentedCode - Whether to add shader definitions as comments.
- * @returns true if the shader is preprocessed without error.
- */
-bool PreprocessShader(
+bool InnerPreprocessShaderMcpp(
 	FString& OutPreprocessedShader,
 	FShaderCompilerOutput& ShaderOutput,
 	const FShaderCompilerInput& ShaderInput,
-	const FShaderCompilerDefinitions& AdditionalDefines,
-	EDumpShaderDefines DefinesPolicy
-	)
+	const FShaderCompilerDefinitions& AdditionalDefines)
 {
-	TRACE_CPUPROFILER_EVENT_SCOPE(PreprocessShader);
-
-	// Skip the cache system and directly load the file path (used for debugging)
-	if (ShaderInput.bSkipPreprocessedCache)
-	{
-		return FFileHelper::LoadFileToString(OutPreprocessedShader, *ShaderInput.VirtualSourceFilePath);
-	}
-	else
-	{
-		check(CheckVirtualShaderFilePath(ShaderInput.VirtualSourceFilePath));
-	}
-
 	int32 McppResult = 0;
 	FString McppOutput, McppErrors;
 
@@ -313,22 +309,366 @@ bool PreprocessShader(
 
 	if (!bHasIncludedMandatoryHeaders)
 	{
-		FShaderCompilerError Error;
-		Error.ErrorVirtualFilePath = ShaderInput.VirtualSourceFilePath;
-		Error.ErrorLineString = TEXT("1");
-		Error.StrippedErrorMessage = TEXT("Error: Shader is required to include /Engine/Public/Platform.ush");
-
-		ShaderOutput.Errors.Add(Error);
+		LogMandatoryHeaderError(ShaderInput, ShaderOutput);
 		return false;
 	}
 
+	OutPreprocessedShader += McppOutput;
+
+	return true;
+}
+
+extern "C"
+{
+	// adapter functions for STB memory allocation
+	void* StbMalloc(size_t Size) 
+	{
+		void* Alloc = FMemory::Malloc(Size);
+		return Alloc; 
+	}
+
+	void* StbRealloc(void* Pointer, size_t Size) 
+	{
+		void* Alloc = FMemory::Realloc(Pointer, Size);
+		return Alloc;
+	}
+
+	void StbFree(void* Pointer) 
+	{
+		return FMemory::Free(Pointer); 
+	}
+
+	ANSICHAR* StbStrDup(const ANSICHAR* InString)
+	{
+		if (InString)
+		{
+			int32 Len = FCStringAnsi::Strlen(InString) + 1;
+			ANSICHAR* Result = reinterpret_cast<ANSICHAR*>(StbMalloc(Len));
+			return FCStringAnsi::Strncpy(Result, InString, Len);
+		}
+		return nullptr;
+	}
+}
+
+struct FStbPreprocessContext
+{
+	const FShaderCompilerInput& ShaderInput;
+	TMap<FString, TArray<ANSICHAR>> LoadedIncludesCache;
+	TMap<FString, TUniquePtr<ANSICHAR[]>> SeenPathsLowerCase;
+
+	bool HasIncludedMandatoryHeaders()
+	{
+		return SeenPathsLowerCase.Contains(PlatformHeaderLowerCase);
+	}
+};
+
+inline bool IsEndOfLine(ANSICHAR C)
+{
+	return C == '\r' || C == '\n';
+}
+
+inline bool CommentStripNeedsHandling(ANSICHAR C)
+{
+	return IsEndOfLine(C) || C == '/' || C == 0;
+}
+
+inline int NewlineCharCount(ANSICHAR First, ANSICHAR Second)
+{
+	return ((First + Second) == '\r' + '\n') ? 2 : 1;
+}
+
+void ConvertAndStripComments(const FString& ShaderSource, TArray<ANSICHAR>& OutStripped)
+{
+	auto ShaderSourceAnsiConvert = StringCast<ANSICHAR>(*ShaderSource);
+
+	// STB preprocessor does not strip comments, so we do so here before returning the loaded source
+	// Doing so is barely more costly than the memcopy we require anyways so has negligible overhead.
+	// Reserve worst case (i.e. assuming there are no comments at all) to avoid reallocation
+	// Note: there's a potential future optimization here if we convert and strip at the same time;
+	// currently this is incurring an extra heap allocation and copy in the case where the StringCast
+	// is not a straight pointer copy (one alloc for the conversion and another for the stripped char array).
+	OutStripped.SetNumUninitialized(ShaderSourceAnsiConvert.Length() + 1); // +1 to append null terminator
+
+	ANSICHAR* CurrentOut = OutStripped.GetData();
+
+	const ANSICHAR* const End = ShaderSourceAnsiConvert.Get() + ShaderSourceAnsiConvert.Length();
+	for (const ANSICHAR* Current = ShaderSourceAnsiConvert.Get(); Current < End;)
+	{
+		while (!CommentStripNeedsHandling(*Current))
+		{
+			*CurrentOut++ = *Current++;
+		}
+
+		if (IsEndOfLine(*Current))
+		{
+			*CurrentOut++ = '\n';
+			Current += NewlineCharCount(Current[0], Current[1]);
+		}
+		else if (Current[0] == '/')
+		{
+			if (Current[1] == '/')
+			{
+				while (!IsEndOfLine(*Current) && Current < End)
+				{
+					++Current;
+				}
+			}
+			else if (Current[1] == '*')
+			{
+				Current += 2;
+				while (!(Current[0] == '*' && Current[1] == '/'))
+				{
+					if (IsEndOfLine(*Current))
+					{
+						*CurrentOut++ = '\n';
+						Current += NewlineCharCount(Current[0], Current[1]);
+					}
+					else
+					{
+						++Current;
+					}
+				}
+				Current += 2;
+			}
+			else
+			{
+				*CurrentOut++ = *Current++;
+			}
+		}
+	}
+	// Null terminate after comment-stripped copy
+	*CurrentOut++ = 0;
+
+	// Set correct length after stripping but don't bother shrinking/reallocating, minor memory overhead to save time
+	OutStripped.SetNum(CurrentOut - OutStripped.GetData(), /* bAllowShrinking */false);
+}
+
+static const ANSICHAR* StbLoadFile(const ANSICHAR* Filename, void* RawContext, size_t* OutLength)
+{
+	FStbPreprocessContext& Context = *reinterpret_cast<FStbPreprocessContext*>(RawContext);
+	FString FilenameConverted = StringCast<TCHAR>(Filename).Get();
+	TArray<ANSICHAR>* ContentsCached = Context.LoadedIncludesCache.Find(FilenameConverted);
+	if (!ContentsCached)
+	{
+		FString ShaderSource;
+
+		if (Context.ShaderInput.Environment.IncludeVirtualPathToContentsMap.Contains(FilenameConverted))
+		{
+			ShaderSource = Context.ShaderInput.Environment.IncludeVirtualPathToContentsMap.FindRef(FilenameConverted);
+		}
+		else if (Context.ShaderInput.Environment.IncludeVirtualPathToExternalContentsMap.Contains(FilenameConverted))
+		{
+			ShaderSource = *Context.ShaderInput.Environment.IncludeVirtualPathToExternalContentsMap.FindRef(FilenameConverted);
+		}
+		else
+		{
+			CheckShaderHashCacheInclude(FilenameConverted, Context.ShaderInput.Target.GetPlatform(), Context.ShaderInput.ShaderFormat.ToString());
+			LoadShaderSourceFile(*FilenameConverted, Context.ShaderInput.Target.GetPlatform(), &ShaderSource, nullptr);
+		}
+		if (!ShaderSource.IsEmpty())
+		{
+			ContentsCached = &Context.LoadedIncludesCache.Add(FilenameConverted);
+			ConvertAndStripComments(ShaderSource, *ContentsCached);
+		}
+	}
+	*OutLength = ContentsCached->Num();
+	return ContentsCached->GetData();
+}
+
+static void StbFreeFile(const ANSICHAR* Filename, const ANSICHAR* Contents, void* RawContext)
+{
+	FStbPreprocessContext& Context = *reinterpret_cast<FStbPreprocessContext*>(RawContext);
+	FString FilenameConverted = StringCast<TCHAR>(Filename).Get();
+	Context.LoadedIncludesCache.FindAndRemoveChecked(FilenameConverted);
+}
+
+static const ANSICHAR* StbResolveInclude(const ANSICHAR* PathInSource, uint32 PathLen, const ANSICHAR* ParentPathAnsi, void* RawContext)
+{
+	FStbPreprocessContext& Context = *reinterpret_cast<FStbPreprocessContext*>(RawContext);
+	FString PathModified(PathLen, PathInSource);
+	FString ParentFolder(ParentPathAnsi);
+	ParentFolder = FPaths::GetPath(ParentFolder);
+	if (!PathModified.StartsWith(TEXT("/"))) // if path doesn't start with / it's relative, if so append the parent's folder and collapse any relative dirs
+	{
+		PathModified = ParentFolder / PathModified;
+		FPaths::CollapseRelativeDirectories(PathModified);
+	}
+
+	// Substitute virtual platform path here to make sure that #line directives refer to the platform-specific file.
+	ReplaceVirtualFilePathForShaderPlatform(PathModified, Context.ShaderInput.Target.GetPlatform());
+
+	// Fixup autogen file
+	ReplaceVirtualFilePathForShaderAutogen(PathModified, Context.ShaderInput.Target.GetPlatform(), &Context.ShaderInput.ShaderPlatformName);
+
+	FString PathModifiedLowerCase = PathModified.ToLower();
+	const TUniquePtr<ANSICHAR[]>* SeenPath = Context.SeenPathsLowerCase.Find(PathModifiedLowerCase);
+	// Keep track of previously resolved paths in a case insensitive manner so preprocessor will handle #pragma once with files included with inconsistent casing correctly
+	// (we store the first correctly resolved path with original casing so we get "nice" line directives)
+	if (SeenPath)
+	{
+		return SeenPath->Get();
+	}
+
+	bool bExists =
+		Context.ShaderInput.Environment.IncludeVirtualPathToContentsMap.Contains(PathModified) ||
+		Context.ShaderInput.Environment.IncludeVirtualPathToExternalContentsMap.Contains(PathModified) ||
+		// LoadShaderSourceFile will load the file if it exists, but then cache it internally, so the next call in StbLoadFile will be cheap
+		// (and hence this is not overly wasteful)
+		LoadShaderSourceFile(*PathModified, Context.ShaderInput.Target.GetPlatform(), nullptr, nullptr);
+
+	if (bExists)
+	{
+		int32 Length = FPlatformString::ConvertedLength<ANSICHAR>(*PathModified);
+		TUniquePtr<ANSICHAR[]>& OutPath = Context.SeenPathsLowerCase.Add(PathModifiedLowerCase, MakeUnique<ANSICHAR[]>(Length));
+		FPlatformString::Convert<TCHAR, ANSICHAR>(OutPath.Get(), Length, *PathModified);
+		return OutPath.Get();
+	}
+
+	return nullptr;
+}
+
+class FShaderPreprocessorModule : public IModuleInterface
+{
+	virtual void StartupModule() override
+	{
+		init_preprocessor(&StbLoadFile, &StbFreeFile, &StbResolveInclude);
+		// disable the "directive not at start of line" error; this allows a few things:
+		// 1. #define'ing #pragma messages - consumed by the preprocessor (to handle UESHADERMETADATA hackery)
+		// 2. #define'ing other #pragmas (those not processed explicitly by the preprocessor are copied into the preprocessed code
+		// 3. handling the HLSL infinity constant (1.#INF); STB preprocessor interprets any use of # as a directive which is not the case here
+		pp_set_warning_mode(PP_RESULT_directive_not_at_start_of_line, PP_RESULT_MODE_no_warning); 
+	}
+};
+IMPLEMENT_MODULE(FShaderPreprocessorModule, ShaderPreprocessor);
+
+static void AddStbDefine(stb_arena* MacroArena, macro_definition**& StbDefines, const TCHAR* Name, const TCHAR* Value)
+{
+	FString Define(FString::Printf(TEXT("%s %s"), Name, Value));
+	auto ConvertedDefine = StringCast<ANSICHAR>(*Define);
+	arrput(StbDefines, pp_define(MacroArena, (ANSICHAR*)ConvertedDefine.Get()));
+}
+
+static void AddStbDefines(stb_arena* MacroArena, macro_definition**& StbDefines, TMap<FString, FString> DefinitionsMap)
+{
+	for (TMap<FString, FString>::TConstIterator It(DefinitionsMap); It; ++It)
+	{
+		AddStbDefine(MacroArena, StbDefines, *It.Key(), *It.Value());
+	}
+}
+
+bool InnerPreprocessShaderStb(
+	FString& OutPreprocessedShader,
+	FShaderCompilerOutput& ShaderOutput,
+	const FShaderCompilerInput& ShaderInput,
+	const FShaderCompilerDefinitions& AdditionalDefines
+)
+{
+	stb_arena MacroArena = { 0 };
+	macro_definition** StbDefines = nullptr;
+
+	AddStbDefines(&MacroArena, StbDefines, ShaderInput.Environment.GetDefinitions());
+	AddStbDefines(&MacroArena, StbDefines, AdditionalDefines.GetDefinitionMap());
+	AddStbDefine(&MacroArena, StbDefines, TEXT("_STB_PREPROCESS"), TEXT("1"));
+
+	FStbPreprocessContext Context{ ShaderInput };
+
+	auto InFilename = StringCast<ANSICHAR>(*ShaderInput.VirtualSourceFilePath);
+	int NumDiagnostics = 0;
+	pp_diagnostic* Diagnostics = nullptr;
+	
+	char* OutPreprocessedAnsi = preprocess_file(nullptr, InFilename.Get(), &Context, StbDefines, arrlen(StbDefines), &Diagnostics, &NumDiagnostics);
+	bool HasError = false;
+	if (Diagnostics != nullptr)
+	{
+		for (int DiagIndex = 0; DiagIndex < NumDiagnostics; ++DiagIndex)
+		{
+			pp_diagnostic* Diagnostic = &Diagnostics[DiagIndex];
+			HasError |= (Diagnostic->error_level == PP_RESULT_MODE_error);
+			
+			FString Message = Diagnostic->message;
+			// as we do with MCPP, we are ignoring warnings (for now?)
+			if (Diagnostic->error_level == PP_RESULT_MODE_error)
+			{
+				FShaderCompilerError* CompilerError = new(ShaderOutput.Errors) FShaderCompilerError;
+				CompilerError->ErrorVirtualFilePath = Diagnostic->where->filename;
+				CompilerError->ErrorLineString = FString::Printf(TEXT("%d"), Diagnostic->where->line_number);
+				CompilerError->StrippedErrorMessage = Message;
+			}
+			else
+			{
+				EMessageType Type = FilterPreprocessorError(Message);
+				if (Type == EMessageType::ShaderMetaData)
+				{
+					FString Directive;
+					ExtractDirective(Directive, Message);
+					ShaderOutput.PragmaDirectives.Add(Directive);
+				}
+			}
+		}
+	}
+
+	OutPreprocessedShader = StringCast<TCHAR>(OutPreprocessedAnsi).Get();
+
+	if (!HasError && !Context.HasIncludedMandatoryHeaders())
+	{
+		LogMandatoryHeaderError(ShaderInput, ShaderOutput);
+		HasError = true;
+	}
+
+	preprocessor_file_free(OutPreprocessedAnsi, Diagnostics);
+	stbds_arrfree(StbDefines);
+	stb_arena_free(&MacroArena);
+
+	return !HasError;
+}
+
+/**
+ * Preprocess a shader.
+ * @param OutPreprocessedShader - Upon return contains the preprocessed source code.
+ * @param ShaderOutput - ShaderOutput to which errors can be added.
+ * @param ShaderInput - The shader compiler input.
+ * @param AdditionalDefines - Additional defines with which to preprocess the shader.
+ * @param DefinesPolicy - Whether to add shader definitions as comments.
+ * @returns true if the shader is preprocessed without error.
+ */
+bool PreprocessShader(
+	FString& OutPreprocessedShader,
+	FShaderCompilerOutput& ShaderOutput,
+	const FShaderCompilerInput& ShaderInput,
+	const FShaderCompilerDefinitions& AdditionalDefines,
+	EDumpShaderDefines DefinesPolicy
+)
+{
+	TRACE_CPUPROFILER_EVENT_SCOPE(PreprocessShader);
+
+	// Skip the cache system and directly load the file path (used for debugging)
+	if (ShaderInput.bSkipPreprocessedCache)
+	{
+		return FFileHelper::LoadFileToString(OutPreprocessedShader, *ShaderInput.VirtualSourceFilePath);
+	}
+	else
+	{
+		check(CheckVirtualShaderFilePath(ShaderInput.VirtualSourceFilePath));
+	}
+
+	bool bResult = false;
+	bool bLegacyPreprocess = ShaderInput.Environment.CompilerFlags.Contains(CFLAG_UseLegacyPreprocessor);
+	FString PreprocessorOutput;
+	if (!bLegacyPreprocess)
+	{
+		bResult |= InnerPreprocessShaderStb(PreprocessorOutput, ShaderOutput, ShaderInput, AdditionalDefines);
+	}
+	else
+	{
+		bResult |= InnerPreprocessShaderMcpp(PreprocessorOutput, ShaderOutput, ShaderInput, AdditionalDefines);
+	}
+
 	// List the defines used for compilation in the preprocessed shaders, especially to know witch permutation vector this shader is.
 	if (DefinesPolicy == EDumpShaderDefines::AlwaysIncludeDefines || (DefinesPolicy == EDumpShaderDefines::DontCare && ShaderInput.DumpDebugInfoPath.Len() > 0))
 	{
 		DumpShaderDefinesAsCommentedCode(ShaderInput, &OutPreprocessedShader);
 	}
 
-	OutPreprocessedShader += McppOutput;
-
-	return true;
+	OutPreprocessedShader += PreprocessorOutput;
+	return bResult;
 }
diff --git a/Engine/Source/Developer/ShaderPreprocessor/Private/StbConfig.h b/Engine/Source/Developer/ShaderPreprocessor/Private/StbConfig.h
new file mode 100644
--- /dev/null
+++ b/Engine/Source/Developer/ShaderPreprocessor/Private/StbConfig.h
@@ -0,0 +1,20 @@
+// Copyright Epic Games, Inc. All Rights Reserved.
+
+#pragma once
+
+#ifdef __cplusplus
+extern "C"
+{
+#endif
+void* StbMalloc(size_t Size);
+void* StbRealloc(void* Pointer, size_t Size);
+void StbFree(void* Pointer);
+char* StbStrDup(const char* String);
+#ifdef __cplusplus
+}
+#endif
+
+#define STB_COMMON_MALLOC(Size) StbMalloc(Size)
+#define STB_COMMON_REALLOC(Pointer, Size) StbRealloc(Pointer, Size)
+#define STB_COMMON_FREE(Pointer) StbFree(Pointer)
+#define STB_COMMON_STRDUP(String) StbStrDup(String)
diff --git a/Engine/Source/Developer/ShaderPreprocessor/ShaderPreprocessor.Build.cs b/Engine/Source/Developer/ShaderPreprocessor/ShaderPreprocessor.Build.cs
--- a/Engine/Source/Developer/ShaderPreprocessor/ShaderPreprocessor.Build.cs
+++ b/Engine/Source/Developer/ShaderPreprocessor/ShaderPreprocessor.Build.cs
@@ -15,7 +15,6 @@ public class ShaderPreprocessor : ModuleRules
 
 		AddEngineThirdPartyPrivateStaticDependencies(Target, "MCPP");
 
-		PrivateDefinitions.Add("STB_DS_IMPLEMENTATION");
-		PrivateDefinitions.Add("STB_ALLOC_IMPLEMENTATION");
+		PrivateDefinitions.Add("STB_CONFIG=../StbConfig.h");
 	}
 }
diff --git a/Engine/Source/Programs/ShaderCompileWorker/Private/ShaderCompileWorker.cpp b/Engine/Source/Programs/ShaderCompileWorker/Private/ShaderCompileWorker.cpp
--- a/Engine/Source/Programs/ShaderCompileWorker/Private/ShaderCompileWorker.cpp
+++ b/Engine/Source/Programs/ShaderCompileWorker/Private/ShaderCompileWorker.cpp
@@ -1300,6 +1300,9 @@ static int32 GuardedMain(int32 argc, TCHAR* argv[], bool bDirectMode)
 	AtomicSave(IniBootstrapFilename,     [](const FString& TmpFile) { GConfig->SaveCurrentStateForBootstrap(*TmpFile); });
 	AtomicSave(ModulesBootstrapFilename, [](const FString& TmpFile) { FModuleManager::Get().SaveCurrentStateForBootstrap(*TmpFile); });
 
+	// Explicitly load ShaderPreprocessor module so it will run its initialization step
+	FModuleManager::LoadModuleChecked<IModuleInterface>(TEXT("ShaderPreprocessor"));
+
 	// We just enumerate the shader formats here for debugging.
 	const TArray<const class IShaderFormat*>& ShaderFormats = GetShaderFormats();
 	check(ShaderFormats.Num());
diff --git a/Engine/Source/Runtime/Engine/Private/ShaderCompiler/ShaderCompiler.cpp b/Engine/Source/Runtime/Engine/Private/ShaderCompiler/ShaderCompiler.cpp
--- a/Engine/Source/Runtime/Engine/Private/ShaderCompiler/ShaderCompiler.cpp
+++ b/Engine/Source/Runtime/Engine/Private/ShaderCompiler/ShaderCompiler.cpp
@@ -1209,6 +1209,16 @@ static TAutoConsoleVariable<int32> CVarShadersRemoveDeadCode(
 	TEXT("\t1: Remove unreferenced code before compilation\n"),
 	ECVF_ReadOnly);
 
+// note: this cvar is intended to be short-lived; when removing be sure to remove the associated code from
+// ShaderMapAppendKeyString (in Shader.cpp)
+static TAutoConsoleVariable<bool> CVarShadersUseLegacyPreprocessor(
+	TEXT("r.Shaders.UseLegacyPreprocessor"),
+	true,
+	TEXT("Executes shader preprocessing via the legacy MCPP preprocessor (instead of the new STB preprocessor).\n")
+	TEXT("\ttrue: Enabled - preprocess with MCPP\n")
+	TEXT("\tfalse: Disabled - preprocess with STB\n"),
+	ECVF_ReadOnly);
+
 extern bool CompileShaderPipeline(const IShaderFormat* Compiler, FName Format, FShaderPipelineCompileJob* PipelineJob, const FString& Dir);
 
 #if ENABLE_COOK_STATS
@@ -3736,6 +3746,9 @@ FShaderCompilingManager::FShaderCompilingManager() :
 	verify(GConfig->GetBool( TEXT("DevOptions.Shaders"), TEXT("bAllowCompilingThroughWorkers"), bAllowCompilingThroughWorkers, GEngineIni ));
 	verify(GConfig->GetBool( TEXT("DevOptions.Shaders"), TEXT("bAllowAsynchronousShaderCompiling"), bAllowAsynchronousShaderCompiling, GEngineIni ));
 
+	// Explicitly load ShaderPreprocessor module so it will run its initialization step
+	FModuleManager::LoadModuleChecked<IModuleInterface>(TEXT("ShaderPreprocessor"));
+	
 	// override the use of workers, can be helpful for debugging shader compiler code
 	static const IConsoleVariable* CVarAllowCompilingThroughWorkers = IConsoleManager::Get().FindConsoleVariable(TEXT("r.Shaders.AllowCompilingThroughWorkers"), false);
 	if (!FPlatformProcess::SupportsMultithreading() || FParse::Param(FCommandLine::Get(), TEXT("noshaderworker")) || (CVarAllowCompilingThroughWorkers && CVarAllowCompilingThroughWorkers->GetInt() == 0))
@@ -6459,6 +6472,11 @@ void GlobalBeginCompileShader(
 		Input.Environment.CompilerFlags.Add(CFLAG_RemoveDeadCode);
 	}
 
+	if (CVarShadersUseLegacyPreprocessor.GetValueOnAnyThread())
+	{
+		Input.Environment.CompilerFlags.Add(CFLAG_UseLegacyPreprocessor);
+	}
+
 	{
 		static const auto CVar = IConsoleManager::Get().FindConsoleVariable(TEXT("r.VT.AnisotropicFiltering"));
 		Input.Environment.SetDefine(TEXT("VIRTUAL_TEXTURE_ANISOTROPIC_FILTERING"), CVar ? (CVar->GetInt() != 0) : 0);
diff --git a/Engine/Source/Runtime/RenderCore/Private/Shader.cpp b/Engine/Source/Runtime/RenderCore/Private/Shader.cpp
--- a/Engine/Source/Runtime/RenderCore/Private/Shader.cpp
+++ b/Engine/Source/Runtime/RenderCore/Private/Shader.cpp
@@ -2053,6 +2053,14 @@ void ShaderMapAppendKeyString(EShaderPlatform Platform, FString& KeyString)
 		}
 	}
 
+	{
+		static const auto CVar = IConsoleManager::Get().FindTConsoleVariableDataBool(TEXT("r.Shaders.UseLegacyPreprocessor"));
+		if (CVar && CVar->GetValueOnAnyThread() != 0)
+		{
+			KeyString += TEXT("_MCPP");
+		}
+	}
+
 	if (RHISupportsRenderTargetWriteMask(Platform))
 	{
 		KeyString += TEXT("_RTWM");
diff --git a/Engine/Source/Runtime/RenderCore/Public/ShaderCompilerCore.h b/Engine/Source/Runtime/RenderCore/Public/ShaderCompilerCore.h
--- a/Engine/Source/Runtime/RenderCore/Public/ShaderCompilerCore.h
+++ b/Engine/Source/Runtime/RenderCore/Public/ShaderCompilerCore.h
@@ -120,8 +120,8 @@ enum ECompilerFlags
 	CFLAG_BindlessSamplers,
 	// EXPERIMENTAL: Run the shader re-writer that removes any unused functions/resources/types from source code before compilation.
 	CFLAG_RemoveDeadCode,
-	// EXPERIMENTAL: Use the DXC preprocessor (instead of the legacy MCPP)
-	CFLAG_PreprocessWithDXC,
+	// Execute shader preprocessing with MCPP (instead of the new STB preprocessor)
+	CFLAG_UseLegacyPreprocessor,
 
 	CFLAG_Max,
 };
