[shaders] fix msvc static analysis warning - enforce that loaded includes are non-empty when loading via stb file load callback.

diff --git a/Engine/Source/Developer/ShaderPreprocessor/Private/ShaderPreprocessor.cpp b/Engine/Source/Developer/ShaderPreprocessor/Private/ShaderPreprocessor.cpp
--- a/Engine/Source/Developer/ShaderPreprocessor/Private/ShaderPreprocessor.cpp
+++ b/Engine/Source/Developer/ShaderPreprocessor/Private/ShaderPreprocessor.cpp
@@ -463,12 +463,11 @@ static const ANSICHAR* StbLoadFile(const ANSICHAR* Filename, void* RawContext, s
 			CheckShaderHashCacheInclude(FilenameConverted, Context.ShaderInput.Target.GetPlatform(), Context.ShaderInput.ShaderFormat.ToString());
 			LoadShaderSourceFile(*FilenameConverted, Context.ShaderInput.Target.GetPlatform(), &ShaderSource, nullptr);
 		}
-		if (!ShaderSource.IsEmpty())
-		{
-			ContentsCached = &Context.LoadedIncludesCache.Add(FilenameConverted);
-			ConvertAndStripComments(ShaderSource, *ContentsCached);
-		}
+		check(!ShaderSource.IsEmpty());
+		ContentsCached = &Context.LoadedIncludesCache.Add(FilenameConverted);
+		ConvertAndStripComments(ShaderSource, *ContentsCached);
 	}
+	check(ContentsCached);
 	*OutLength = ContentsCached->Num();
 	return ContentsCached->GetData();
 }
