STB processor reemits line directives.

diff --git a/Engine/Source/Developer/ShaderPreprocessor/Private/stb_preprocess/preprocessor.c b/Engine/Source/Developer/ShaderPreprocessor/Private/stb_preprocess/preprocessor.c
--- a/Engine/Source/Developer/ShaderPreprocessor/Private/stb_preprocess/preprocessor.c
+++ b/Engine/Source/Developer/ShaderPreprocessor/Private/stb_preprocess/preprocessor.c
@@ -3583,9 +3583,25 @@ static void process_directive(parse_state* cs, conditional_state* cons)
 		}
 
 		case HASH_line:
-			// @TODO
-			if (alloc)
-				arrfree(alloc);
+			// copy line directive through to output... existing code will output the newline, so don't do it here
+			if (cs->fast_dest)
+			{
+				size_t out_length = strlen(p);
+				char* z = cs->fast_dest;
+				memcpy(z, "#line ", 6);
+				z += 6;
+				memcpy(z, p, out_length);
+				z += out_length;
+				cs->fast_dest = z;
+			}
+			else
+			{
+				size_t out_length = strlen(p);
+				size_t s = arraddnindex(cs->dest, 6);
+				memcpy(cs->dest + s, "#line ", 6);
+				s = arraddnindex(cs->dest, out_length);
+				memcpy(cs->dest + s, p, out_length);
+			}
 			return;
 
 		case HASH_pragma:
