Disable sndbs distribution for ShaderPreprocessor module

#rnx

diff --git a/Engine/Source/Developer/ShaderPreprocessor/ShaderPreprocessor.Build.cs b/Engine/Source/Developer/ShaderPreprocessor/ShaderPreprocessor.Build.cs
--- a/Engine/Source/Developer/ShaderPreprocessor/ShaderPreprocessor.Build.cs
+++ b/Engine/Source/Developer/ShaderPreprocessor/ShaderPreprocessor.Build.cs
@@ -16,5 +16,8 @@ public class ShaderPreprocessor : ModuleRules
 		AddEngineThirdPartyPrivateStaticDependencies(Target, "MCPP");
 
 		PrivateDefinitions.Add("STB_CONFIG=../StbConfig.h");
+
+		// TODO: Disable SN-DBS distribution until ../StbConfig.h include from STB_CONFIG can be resolved correctly
+		bBuildLocallyWithSNDBS = true;
 	}
 }
