ShaderPreprocessor: Fix access violation due to reading off the end of the input text when the text contains an unterminated /* comment.
#rnx

diff --git a/Engine/Source/Developer/ShaderPreprocessor/Private/ShaderPreprocessor.cpp b/Engine/Source/Developer/ShaderPreprocessor/Private/ShaderPreprocessor.cpp
--- a/Engine/Source/Developer/ShaderPreprocessor/Private/ShaderPreprocessor.cpp
+++ b/Engine/Source/Developer/ShaderPreprocessor/Private/ShaderPreprocessor.cpp
@@ -390,8 +390,11 @@ void ConvertAndStripComments(const FString& ShaderSource, TArray<ANSICHAR>& OutS
 	ANSICHAR* CurrentOut = OutStripped.GetData();
 
 	const ANSICHAR* const End = ShaderSourceAnsiConvert.Get() + ShaderSourceAnsiConvert.Length();
+	// We rely on null termination to avoid the need to check Current < End in some cases
+	check(*End == '\0');
 	for (const ANSICHAR* Current = ShaderSourceAnsiConvert.Get(); Current < End;)
 	{
+		// CommentStripNeedsHandling returns true when *Current == '\0;
 		while (!CommentStripNeedsHandling(*Current))
 		{
 			*CurrentOut++ = *Current++;
@@ -414,9 +417,14 @@ void ConvertAndStripComments(const FString& ShaderSource, TArray<ANSICHAR>& OutS
 			else if (Current[1] == '*')
 			{
 				Current += 2;
-				while (!(Current[0] == '*' && Current[1] == '/'))
+				while (Current < End)
 				{
-					if (IsEndOfLine(*Current))
+					if (Current[0] == '*' && Current[1] == '/')
+					{
+						Current += 2;
+						break;
+					}
+					else if (IsEndOfLine(*Current))
 					{
 						*CurrentOut++ = '\n';
 						Current += NewlineCharCount(Current[0], Current[1]);
@@ -426,7 +434,6 @@ void ConvertAndStripComments(const FString& ShaderSource, TArray<ANSICHAR>& OutS
 						++Current;
 					}
 				}
-				Current += 2;
 			}
 			else
 			{
