[shader preprocessor] implement #error directives

diff --git a/Engine/Source/Developer/ShaderPreprocessor/Private/stb_preprocess/preprocessor.c b/Engine/Source/Developer/ShaderPreprocessor/Private/stb_preprocess/preprocessor.c
--- a/Engine/Source/Developer/ShaderPreprocessor/Private/stb_preprocess/preprocessor.c
+++ b/Engine/Source/Developer/ShaderPreprocessor/Private/stb_preprocess/preprocessor.c
@@ -477,11 +477,37 @@ void fill_where(parse_state* stack, pp_diagnostic* diagnostic, int* line_number)
 	arrpush(diagnostic->where, where);
 }
 
-static int error_explicit_v(parse_state* ps, int code, int line_number, char* text, va_list va)
+static int error_explicit(parse_state* ps, int code, int line_number, char* text)
 {
 	STB_ASSUME(ps != NULL);
 	parse_state* stack;
 	pp_diagnostic d;
+
+	d.message = text;
+	d.diagnostic_code = code;
+	d.error_level = pp_result_mode[code];
+	if (d.error_level == PP_RESULT_MODE_warning_fast)
+		d.error_level = PP_RESULT_MODE_warning;
+	d.where = 0;
+
+	fill_where(ps, &d, &line_number);
+	stack = ps->parent;
+	while (stack != NULL)
+	{
+		fill_where(stack, &d, &line_number);
+		stack = stack->parent;
+	}
+
+	arrpush(ps->context->error, d);
+
+	if (d.error_level == PP_RESULT_MODE_error)
+		ps->context->stop = 1;
+
+	return d.error_level == PP_RESULT_MODE_error;
+}
+
+static int error_explicit_v(parse_state* ps, int code, int line_number, char* text, va_list va)
+{
 	size_t len;
 
 #if defined(_MSC_VER) && _MSC_VER < 1500
@@ -504,36 +530,12 @@ static int error_explicit_v(parse_state* ps, int code, int line_number, char* te
 	if (size < 1)
 		return 1;
 
-	d.message = (char*)STB_COMMON_MALLOC(size);
-
-	// sanity check added to satisfy clang static analyzer
-	if (d.message == NULL)
-		return 1;
-
-	len = _vsnprintf(d.message, len, text, va);
+	char* message = (char*)STB_COMMON_MALLOC(size);
+	len = _vsnprintf(message, len, text, va);
 	
-	d.message[size - 1] = 0;
+	message[size - 1] = 0;
 
-	d.diagnostic_code = code;
-	d.error_level = pp_result_mode[code];
-	if (d.error_level == PP_RESULT_MODE_warning_fast)
-		d.error_level = PP_RESULT_MODE_warning;
-	d.where = 0;
-
-	fill_where(ps, &d, &line_number);
-	stack = ps->parent;
-	while (stack != NULL)
-	{
-		fill_where(stack, &d, &line_number);
-		stack = stack->parent;
-	}
-
-	arrpush(ps->context->error, d);
-
-	if (d.error_level == PP_RESULT_MODE_error)
-		ps->context->stop = 1;
-
-	return d.error_level == PP_RESULT_MODE_error;
+	return error_explicit(ps, code, line_number, message);
 }
 
 // returns TRUE if it should terminate
@@ -4092,7 +4094,17 @@ static void process_directive(parse_state* cs, conditional_state* cons)
 		}
 
 		case HASH_error:
-			// @TODO
+			p = preprocessor_skip_whitespace_simple(p);
+			const char* err_end = p;
+			while (*err_end != 0 && !(char_is_end_of_line(*err_end)))
+				err_end++;
+		
+			ptrdiff_t len = (err_end - p) + 1;
+			char* err = (char*)STB_COMMON_MALLOC(len);
+			memcpy(err, p, len - 1);
+			err[len - 1] = 0;
+
+			error_explicit(cs, PP_RESULT_ERROR, 0, err);
 			break;
 	}
 
