[shader preprocessor] optimizations
- keep a preprocessor-specific shared cache of loaded shader files in ANSI format similar to the one in ShaderCore but skipping the unnecessary additional load to widechar, along with the associated extra allocations and conversions, and stripping comments directly as part of the load
- for in-memory source contained in the environment, which can't use the above due to potential different contents of same-named includes across jobs, perform the widechar->ansichar conversion and comment strip in a single step rather than converting then stripping, to save an additional allocation of the full source

diff --git a/Engine/Source/Developer/ShaderPreprocessor/Private/ShaderPreprocessor.cpp b/Engine/Source/Developer/ShaderPreprocessor/Private/ShaderPreprocessor.cpp
--- a/Engine/Source/Developer/ShaderPreprocessor/Private/ShaderPreprocessor.cpp
+++ b/Engine/Source/Developer/ShaderPreprocessor/Private/ShaderPreprocessor.cpp
@@ -4,6 +4,7 @@
 #include "Misc/FileHelper.h"
 #include "Misc/Paths.h"
 #include "Misc/ScopeLock.h"
+#include "Misc/ScopeRWLock.h"
 #include "Modules/ModuleManager.h"
 #include "PreprocessorPrivate.h"
 
@@ -102,40 +103,182 @@ inline bool IsEndOfLine(ANSICHAR C)
 {
 	return C == '\r' || C == '\n';
 }
+inline bool IsEndOfLine(TCHAR C)
+{
+	return C == TEXT('\r') || C == TEXT('\n');
+}
 
 inline bool CommentStripNeedsHandling(ANSICHAR C)
 {
 	return IsEndOfLine(C) || C == '/' || C == 0;
 }
+inline bool CommentStripNeedsHandling(TCHAR C)
+{
+	return IsEndOfLine(C) || C == TEXT('/') || C == 0;
+}
 
 inline int NewlineCharCount(ANSICHAR First, ANSICHAR Second)
 {
 	return ((First + Second) == '\r' + '\n') ? 2 : 1;
 }
+inline int NewlineCharCount(TCHAR First, TCHAR Second)
+{
+	return ((First + Second) == TEXT('\r') + TEXT('\n')) ? 2 : 1;
+}
 
-void ConvertAndStripComments(const FString& ShaderSource, TArray<ANSICHAR>& OutStripped)
+
+inline ANSICHAR ReadChar(FArchive& Ar)
 {
-	auto ShaderSourceAnsiConvert = StringCast<ANSICHAR>(*ShaderSource);
+	if (Ar.AtEnd())
+	{
+		return 0;
+	}
+	ANSICHAR Char;
+	Ar << Char;
+	return Char;
+}
+
+// Given an archive pointing to a shader source file, populates the given array with contents of that source file
+// with all comments stripped. This is needed since the STB preprocessor itself does not strip comments.
+// Used when loading shader files from disk into our internal cache.
+void LoadAndStripComments(FArchive& Ar, TArray<ANSICHAR>& OutStripped)
+{
+	// Reserve worst case (i.e. assuming there are no comments at all) to avoid reallocation
+	int32 BufferSize = Ar.TotalSize() + 1;
+	OutStripped.SetNumUninitialized(BufferSize);
+
+	// empty file case; unlikely in practice but handle it just in case
+	if (Ar.AtEnd())
+	{
+		OutStripped[0] = 0;
+		return;
+	}
+
+	ANSICHAR* CurrentOut = OutStripped.GetData();
+
+	ANSICHAR Char0 = ReadChar(Ar);
+	ANSICHAR Char1 = ReadChar(Ar);
+
+	// skip past UTF8 BOM if it exists
+	if ((uint8)Char0 == 0xef && (uint8)Char1 == 0xbb)
+	{
+		ANSICHAR Char2 = ReadChar(Ar);
+		check((uint8)Char2 == 0xbf);
+
+		Char0 = ReadChar(Ar);
+		Char1 = ReadChar(Ar);
+	}
+
+	while (Char0 != 0)
+	{
+		// sanity check that we're not overrunning the buffer
+		check(CurrentOut < OutStripped.GetData() + BufferSize);
+		if (CommentStripNeedsHandling(Char0))
+		{
+			if (IsEndOfLine(Char0))
+			{
+				*CurrentOut++ = '\n';
+
+				if (NewlineCharCount(Char0, Char1) == 2)
+				{
+					Char0 = ReadChar(Ar);
+					Char1 = ReadChar(Ar);
+					continue;
+				}
+			}
+			else if (Char0 == '/')
+			{
+				// single line comment, skip all chars until we find a newline
+				if (Char1 == '/')
+				{
+					Char0 = ReadChar(Ar);
+					Char1 = ReadChar(Ar);
+					while (Char0 != 0 && !IsEndOfLine(Char0))
+					{
+						Char0 = Char1;
+						Char1 = ReadChar(Ar);
+					}
+					continue;
+				}
+				else if (Char1 == '*')
+				{
+					Char0 = ReadChar(Ar);
+					Char1 = ReadChar(Ar);
+					while (Char0 != 0)
+					{
+						if (Char0 == '*' && Char1 == '/')
+						{
+							Char0 = ReadChar(Ar);
+							Char1 = ReadChar(Ar);
+							break;
+						}
+						// include newlines from multiline comments to preserve line count
+						else if (IsEndOfLine(Char0))
+						{
+							*CurrentOut++ = '\n';
+							if (NewlineCharCount(Char0, Char1) == 2)
+							{
+								Char0 = ReadChar(Ar);
+								Char1 = ReadChar(Ar);
+								continue;
+							}
+						}
+						Char0 = Char1;
+						Char1 = ReadChar(Ar);
+					}
+					continue;
+				}
+			}
+		}
+		
+		*CurrentOut++ = Char0;
+		Char0 = Char1;
+		Char1 = ReadChar(Ar);
+	}
+	// Null terminate after comment-stripped file read
+	check(CurrentOut < OutStripped.GetData() + BufferSize);
+	*CurrentOut++ = 0;
 
+	// Set correct length after stripping but don't bother shrinking/reallocating, minor memory overhead to save time
+	OutStripped.SetNum(CurrentOut - OutStripped.GetData(), /* bAllowShrinking */false);
+}
+
+
+inline ANSICHAR ConvertChecked(TCHAR C)
+{
+	checkf(C >= 0 && C <= 0xff, TEXT("Invalid ANSI char found outside of comment in shader source file"));
+	return (ANSICHAR)C;
+}
+
+// Given an FString containing the contents of a shader source file, populates the given array with contents of
+// that source file with all comments stripped. This is needed since the STB preprocessor itself does not strip 
+// comments.
+// This function is used when reading in-memory source from compile environments.
+void ConvertAndStripComments(const FString& ShaderSource, TArray<ANSICHAR>& OutStripped)
+{
 	// STB preprocessor does not strip comments, so we do so here before returning the loaded source
 	// Doing so is barely more costly than the memcopy we require anyways so has negligible overhead.
 	// Reserve worst case (i.e. assuming there are no comments at all) to avoid reallocation
-	// Note: there's a potential future optimization here if we convert and strip at the same time;
-	// currently this is incurring an extra heap allocation and copy in the case where the StringCast
-	// is not a straight pointer copy (one alloc for the conversion and another for the stripped char array).
-	OutStripped.SetNumUninitialized(ShaderSourceAnsiConvert.Length() + 1); // +1 to append null terminator
+	int32 BufferSize = ShaderSource.Len() + 1; // +1 to append null terminator
+	OutStripped.SetNumUninitialized(BufferSize);
 
 	ANSICHAR* CurrentOut = OutStripped.GetData();
 
-	const ANSICHAR* const End = ShaderSourceAnsiConvert.Get() + ShaderSourceAnsiConvert.Length();
+	const TCHAR* const End = ShaderSource.GetCharArray().GetData() + ShaderSource.Len();
+
 	// We rely on null termination to avoid the need to check Current < End in some cases
-	check(*End == '\0');
-	for (const ANSICHAR* Current = ShaderSourceAnsiConvert.Get(); Current < End;)
+	check(*End == TEXT('\0'));
+	for (const TCHAR* Current = ShaderSource.GetCharArray().GetData(); Current < End;)
 	{
-		// CommentStripNeedsHandling returns true when *Current == '\0;
+		// sanity check that we're not overrunning the buffer
+		check(CurrentOut < (OutStripped.GetData() + BufferSize));
+		// CommentStripNeedsHandling returns true when *Current == '\0';
 		while (!CommentStripNeedsHandling(*Current))
 		{
-			*CurrentOut++ = *Current++;
+			// straight cast to ansichar; since this is a character in hlsl source that's not in a comment
+			// we assume that it must be valid to do so. if this assumption is not valid the shader source was
+			// broken/corrupt anyways.
+			*CurrentOut++ = (ANSICHAR)(*Current++);
 		}
 
 		if (IsEndOfLine(*Current))
@@ -175,61 +318,108 @@ void ConvertAndStripComments(const FString& ShaderSource, TArray<ANSICHAR>& OutS
 			}
 			else
 			{
-				*CurrentOut++ = *Current++;
+				*CurrentOut++ = (ANSICHAR)(*Current++);
 			}
 		}
 	}
 	// Null terminate after comment-stripped copy
+	check(CurrentOut < (OutStripped.GetData() + BufferSize));
 	*CurrentOut++ = 0;
 
 	// Set correct length after stripping but don't bother shrinking/reallocating, minor memory overhead to save time
 	OutStripped.SetNum(CurrentOut - OutStripped.GetData(), /* bAllowShrinking */false);
 }
 
-const FString* FindInMemorySource(const FShaderCompilerEnvironment& Environment, const FString& FilenameConverted)
+constexpr int32 FILECACHE_BUCKETS = 31; 
+struct FPreprocessorShaderFileCache
 {
-	const FString* InMemorySource = Environment.IncludeVirtualPathToContentsMap.Find(FilenameConverted);
-	if (!InMemorySource)
+	FRWLock Lock;
+
+	TMap<FString, TArray<ANSICHAR>> Map;
+} GPreprocessorShaderFileCache[FILECACHE_BUCKETS];
+
+const TArray<ANSICHAR>* LoadShaderFileInternal(const FString& VirtualFilePath)
+{
+	uint32 CurrentHash = GetTypeHash(VirtualFilePath);
+	const TArray<ANSICHAR>* CachedFile = nullptr;
+	FPreprocessorShaderFileCache& ShaderFileCache = GPreprocessorShaderFileCache[CurrentHash % FILECACHE_BUCKETS];
+	{
+		FRWScopeLock ScopeLock(ShaderFileCache.Lock, SLT_ReadOnly);
+		CachedFile = ShaderFileCache.Map.FindByHash(CurrentHash, VirtualFilePath);
+	}
+
+	if (!CachedFile)
 	{
-		const FThreadSafeSharedStringPtr* SharedPtr = Environment.IncludeVirtualPathToExternalContentsMap.Find(FilenameConverted);
-		InMemorySource = SharedPtr ? SharedPtr->Get() : nullptr;
+		FRWScopeLock ScopeLock(ShaderFileCache.Lock, SLT_Write);
+
+		// Double-check the cache while holding exclusive lock as another thread may have added the item we're looking for
+		CachedFile = ShaderFileCache.Map.FindByHash(CurrentHash, VirtualFilePath);
+
+		if (!CachedFile)
+		{
+			FString ShaderFilePath = GetShaderSourceFilePath(VirtualFilePath);
+			FArchive* Reader = IFileManager::Get().CreateFileReader(*ShaderFilePath);
+			if (Reader)
+			{
+				TArray<ANSICHAR>& FileArray = ShaderFileCache.Map.AddByHash(CurrentHash, VirtualFilePath);
+				check(!ShaderFilePath.IsEmpty());
+				LoadAndStripComments(*Reader, FileArray);
+
+				CachedFile = &FileArray;
+				delete Reader;
+			}
+		}
 	}
-	return InMemorySource;
+	return CachedFile;
 }
 
-static const ANSICHAR* StbLoadFile(const ANSICHAR* Filename, void* RawContext, size_t* OutLength)
+const TArray<ANSICHAR>* FindInMemorySource(FStbPreprocessContext& Context, const FString& FilenameConverted)
 {
-	FStbPreprocessContext& Context = *reinterpret_cast<FStbPreprocessContext*>(RawContext);
-	FString FilenameConverted = StringCast<TCHAR>(Filename).Get();
+	// Note: we can't use the shared cache for this path as different environments may have different in-memory source
+	// for the same virtual file path; we instead use a separate cache in the context to save time/memory in the event
+	// the same in-memory include is included in multiple places in the same shader.
 	TArray<ANSICHAR>* ContentsCached = Context.LoadedIncludesCache.Find(FilenameConverted);
 	if (!ContentsCached)
 	{
-		ContentsCached = &Context.LoadedIncludesCache.Add(FilenameConverted);
-		// Local FString used for the LoadShaderSourceFile path; we should consider retrieving source from the shader file cache as a reference
-		// (avoid an extra alloc+copy)
-		FString SourceCopy;
-
-		const FString* InMemorySource = FindInMemorySource(Context.Environment, FilenameConverted);
-
+		const FString* InMemorySource = Context.Environment.IncludeVirtualPathToContentsMap.Find(FilenameConverted);
 		if (!InMemorySource)
 		{
-			CheckShaderHashCacheInclude(FilenameConverted, Context.ShaderInput.Target.GetPlatform(), Context.ShaderInput.ShaderFormat.ToString());
-			LoadShaderSourceFile(*FilenameConverted, Context.ShaderInput.Target.GetPlatform(), &SourceCopy, nullptr);
-			InMemorySource = &SourceCopy;
+			const FThreadSafeSharedStringPtr* SharedPtr = Context.Environment.IncludeVirtualPathToExternalContentsMap.Find(FilenameConverted);
+			InMemorySource = SharedPtr ? SharedPtr->Get() : nullptr;
+		}
+
+		// if we found a match in one of the two maps in the environment, convert to ANSICHAR, strip comments and add to our cache
+		if (InMemorySource)
+		{
+			ContentsCached = &Context.LoadedIncludesCache.Add(FilenameConverted);
+			ConvertAndStripComments(*InMemorySource, *ContentsCached);
 		}
-		check(InMemorySource && !InMemorySource->IsEmpty());
-		ConvertAndStripComments(*InMemorySource, *ContentsCached);
 	}
-	check(ContentsCached);
-	*OutLength = ContentsCached->Num();
-	return ContentsCached->GetData();
+
+	return ContentsCached;
 }
 
-static void StbFreeFile(const ANSICHAR* Filename, const ANSICHAR* Contents, void* RawContext)
+static const ANSICHAR* StbLoadFile(const ANSICHAR* Filename, void* RawContext, size_t* OutLength)
 {
 	FStbPreprocessContext& Context = *reinterpret_cast<FStbPreprocessContext*>(RawContext);
 	FString FilenameConverted = StringCast<TCHAR>(Filename).Get();
-	Context.LoadedIncludesCache.FindAndRemoveChecked(FilenameConverted);
+
+	const TArray<ANSICHAR>* Source = FindInMemorySource(Context, FilenameConverted);
+	if (!Source)
+	{
+		CheckShaderHashCacheInclude(FilenameConverted, Context.ShaderInput.Target.GetPlatform(), Context.ShaderInput.ShaderFormat.ToString());
+		Source = LoadShaderFileInternal(FilenameConverted);
+	}
+	checkf(Source, TEXT("Cannot load shader file %s"), *Filename);
+
+	*OutLength = Source->Num();
+	return Source->GetData();
+}
+
+static void StbFreeFile(const ANSICHAR* Filename, const ANSICHAR* Contents, void* RawContext)
+{
+	// No-op; shader source loaded from the compile environment will be freed from the cache in FStbPreprocessContext when it's destructed, 
+	// and those loaded off disk are intentionally kept around indefinitely in the shared cache.
 }
 
 static const ANSICHAR* StbResolveInclude(const ANSICHAR* PathInSource, uint32 PathLen, const ANSICHAR* ParentPathAnsi, void* RawContext)
@@ -244,11 +434,7 @@ static const ANSICHAR* StbResolveInclude(const ANSICHAR* PathInSource, uint32 Pa
 		FPaths::CollapseRelativeDirectories(PathModified);
 	}
 
-	// Substitute virtual platform path here to make sure that #line directives refer to the platform-specific file.
-	ReplaceVirtualFilePathForShaderPlatform(PathModified, Context.ShaderInput.Target.GetPlatform());
-
-	// Fixup autogen file
-	ReplaceVirtualFilePathForShaderAutogen(PathModified, Context.ShaderInput.Target.GetPlatform(), &Context.ShaderInput.ShaderPlatformName);
+	FixupShaderFilePath(PathModified, Context.ShaderInput.Target.GetPlatform(), &Context.ShaderInput.ShaderPlatformName);
 
 	FString PathModifiedLowerCase = PathModified.ToLower();
 	const TUniquePtr<ANSICHAR[]>* SeenPath = Context.SeenPathsLowerCase.Find(PathModifiedLowerCase);
@@ -262,9 +448,9 @@ static const ANSICHAR* StbResolveInclude(const ANSICHAR* PathInSource, uint32 Pa
 	bool bExists =
 		Context.Environment.IncludeVirtualPathToContentsMap.Contains(PathModified) ||
 		Context.Environment.IncludeVirtualPathToExternalContentsMap.Contains(PathModified) ||
-		// LoadShaderSourceFile will load the file if it exists, but then cache it internally, so the next call in StbLoadFile will be cheap
-		// (and hence this is not overly wasteful)
-		LoadShaderSourceFile(*PathModified, Context.ShaderInput.Target.GetPlatform(), nullptr, nullptr);
+		// LoadShaderFileInternal will load the file if it exists, but then cache it internally, so the next call in StbLoadFile will be cheap
+		// (and hence this is not wasteful, just performs the loading earlier)
+		LoadShaderFileInternal(*PathModified);
 
 	if (bExists)
 	{
