Fixing shader error dumping:
- When checking if we should output debug info, the output won't always be successful
- When preprocessing the shader source, make sure the output buffer is cleared before we do anything to prevent doubling up the shader source on retries.

diff --git a/Engine/Source/Developer/ShaderPreprocessor/Private/ShaderPreprocessor.cpp b/Engine/Source/Developer/ShaderPreprocessor/Private/ShaderPreprocessor.cpp
--- a/Engine/Source/Developer/ShaderPreprocessor/Private/ShaderPreprocessor.cpp
+++ b/Engine/Source/Developer/ShaderPreprocessor/Private/ShaderPreprocessor.cpp
@@ -417,10 +417,10 @@ bool PreprocessShader(
 	{
 		return FFileHelper::LoadFileToString(Output.EditSource(), *Input.VirtualSourceFilePath);
 	}
-	else
-	{
-		check(CheckVirtualShaderFilePath(Input.VirtualSourceFilePath));
-	}
+
+	check(CheckVirtualShaderFilePath(Input.VirtualSourceFilePath));
+
+	Output.EditSource().Empty();
 
 	// List the defines used for compilation in the preprocessed shaders, especially to know which permutation vector this shader is.
 	if (DefinesPolicy == EDumpShaderDefines::AlwaysIncludeDefines || (DefinesPolicy == EDumpShaderDefines::DontCare && Input.DumpDebugInfoPath.Len() > 0))
diff --git a/Engine/Source/Runtime/RenderCore/Private/ShaderCore.cpp b/Engine/Source/Runtime/RenderCore/Private/ShaderCore.cpp
--- a/Engine/Source/Runtime/RenderCore/Private/ShaderCore.cpp
+++ b/Engine/Source/Runtime/RenderCore/Private/ShaderCore.cpp
@@ -1140,7 +1140,7 @@ void CompileShader(const TArray<const IShaderFormat*>& ShaderFormats, FShaderCom
 #endif
 
 	Job.bSucceeded = Job.Output.bSucceeded;
-	if (Job.bSucceeded && Job.Input.DumpDebugInfoEnabled())
+	if (Job.Input.DumpDebugInfoEnabled())
 	{
 		// if the preprocessed cache is disabled, dump debug output here, since we don't serialize preprocess output 
 		// back to the cooker in this case (if enabled this will occur in the job OnComplete callback)
