Build fix -- preprocessor compile error due to assigning a variable twice in one statement (once in the macro, once in the code).

#rnx

diff --git a/Engine/Source/Developer/ShaderPreprocessor/Private/stb_preprocess/preprocessor.c b/Engine/Source/Developer/ShaderPreprocessor/Private/stb_preprocess/preprocessor.c
--- a/Engine/Source/Developer/ShaderPreprocessor/Private/stb_preprocess/preprocessor.c
+++ b/Engine/Source/Developer/ShaderPreprocessor/Private/stb_preprocess/preprocessor.c
@@ -3072,7 +3072,7 @@ static void maybe_expand_macro(parse_state* cs, struct macro_definition* pending
 				// the call to arrsetlen below, because downstream code assumes there is a null terminator beyond the official length of the array.  Doing
 				// the sub-allocation first preserves the null terminator -- without this, the sub-allocation will overwrite it.  We only suballocate if
 				// a minimum of 200 characters is available, as copy_argument reserves this.
-				argument_buffer = stbds_arrinline_suballoc(argument_buffer, char, 200);
+				stbds_arrinline_suballoc(argument_buffer, char, 200);
 
 				arrsetlen(copy, e - p);
 				arrput(arguments, copy);
