Shader Preprocessor optimizations second pass:  10.7% overall improvement to ConditionalPreprocessShader, or 23.1% improvement to lower level PreprocessShader pass alone, factoring out other passes like dead code removal not affected by this CL.

* Low overhead bloom filter for identifier hash lookups in maybe_expand_macro  (3.1%)
* Optimized inline version of library calls (isspace, isalnum), and avoid strtoul for single digit numbers in evaluate_if   (2.6%)
* Array reserve function inlining  (1.7%)
* Fetch output size from preprocessor library instead of calling strlen  (1.2%)
* Avoid FName -> string -> FName round trip conversion of key in FShaderCompilerEnvironment::Merge   (0.3%)
* Use string concatentation instead of printf and inline storage in AddStbDefines   (0.3%)
* Other overhead reduction (remove unused fast_dest feature, smaller stbds_array_header struct, savings from inlining)    (1.5%)

#rnx

diff --git a/Engine/Source/Developer/ShaderPreprocessor/Private/ShaderPreprocessor.cpp b/Engine/Source/Developer/ShaderPreprocessor/Private/ShaderPreprocessor.cpp
--- a/Engine/Source/Developer/ShaderPreprocessor/Private/ShaderPreprocessor.cpp
+++ b/Engine/Source/Developer/ShaderPreprocessor/Private/ShaderPreprocessor.cpp
@@ -227,9 +227,14 @@ IMPLEMENT_MODULE(FShaderPreprocessorModule, ShaderPreprocessor);
 
 static void AddStbDefine(stb_arena* MacroArena, macro_definition**& StbDefines, const TCHAR* Name, const TCHAR* Value)
 {
-	FString Define(FString::Printf(TEXT("%s %s"), Name, Value));
-	auto ConvertedDefine = StringCast<ANSICHAR>(*Define);
-	arrput(StbDefines, pp_define(MacroArena, (ANSICHAR*)ConvertedDefine.Get()));
+	TAnsiStringBuilder<256> Define;
+
+	// Define format:  "%s %s"  (Name Value)
+	Define.Append(Name);
+	Define.AppendChar(' ');
+	Define.Append(Value);
+
+	arrput(StbDefines, pp_define(MacroArena, Define.ToString()));
 }
 
 PRAGMA_DISABLE_DEPRECATION_WARNINGS		// FShaderCompilerDefinitions will be made internal in the future, marked deprecated until then
@@ -271,7 +276,7 @@ bool InnerPreprocessShaderStb(
 		ThreadLocalPreprocessBuffer = new char[ClampedPreprocessBufferSize];
 	}
 
-	char* OutPreprocessedAnsi = preprocess_file(nullptr, InFilename.Get(), &Context, StbDefines, arrlen(StbDefines), &Diagnostics, &NumDiagnostics, ThreadLocalPreprocessBuffer, ClampedPreprocessBufferSize);
+	char* OutPreprocessedAnsi = preprocess_file(InFilename.Get(), &Context, StbDefines, arrlen(StbDefines), &Diagnostics, &NumDiagnostics, ThreadLocalPreprocessBuffer, ClampedPreprocessBufferSize);
 
 	bool HasError = false;
 	if (Diagnostics != nullptr)
@@ -303,7 +308,8 @@ bool InnerPreprocessShaderStb(
 
 	if (!HasError)
 	{
-		Output.EditSource().Append(OutPreprocessedAnsi);
+		// "preprocessor_file_size" includes null terminator, so subtract one for Append call -- passing size saves an expensive strlen in Append
+		Output.EditSource().Append(OutPreprocessedAnsi, preprocessor_file_size(OutPreprocessedAnsi) - 1);
 	}
 
 	if (!HasError && !Context.HasIncludedMandatoryHeaders())
diff --git a/Engine/Source/Developer/ShaderPreprocessor/Private/stb_preprocess/cond_expr.c b/Engine/Source/Developer/ShaderPreprocessor/Private/stb_preprocess/cond_expr.c
--- a/Engine/Source/Developer/ShaderPreprocessor/Private/stb_preprocess/cond_expr.c
+++ b/Engine/Source/Developer/ShaderPreprocessor/Private/stb_preprocess/cond_expr.c
@@ -53,11 +53,78 @@ typedef struct
 	int syntax_error;
 } ppcexp;  // preprocessor constant expression lexer
 
+// Faster inline version of CRT isspace function
+inline unsigned isspace_inline(unsigned p)
+{
+	// Subtract one so ASCII space (32) fits in the last bit of a 32-bit mask.  The null terminator (zero) will wrap around due to
+	// unsigned math, which is good, because it will pass this conditional, avoiding the bit test.
+	p = p - 1;
+	if (p >= 32)
+	{
+		return 0;
+	}
+
+	// Only invalid non-printable characters (which lead to downstream syntax errors) will fail the bit test, since nullptr is covered
+	// by the branch above, so we expect this second branch condition to be 100% correctly predicted.
+	static const unsigned whitespace_mask =
+		(1u << ('\t' - 1)) |
+		(1u << ('\n' - 1)) |
+		(1u << ('\v' - 1)) |
+		(1u << ('\r' - 1)) |
+		(1u << (' '  - 1));
+
+	return whitespace_mask & (1u << p);
+}
+
+// Assuming we found a starting digit, this function checks the next character to see if it ends the number.  If so, it's a single
+// digit number (which is very common in the preprocessor -- most values are zero or one), and we can do a trivial digit character
+// to number conversion, rather than calling the expensive strtoull.  This doesn't try to catch all cases, but it catches most
+// cases likely to exist in valid (non syntax error) code.
+inline int is_end_of_number(int p)
+{
+	// Any ASCII less than a decimal character ('.') ends a number, which includes whitespace, parentheses, some operator
+	// characters, and NULL.  Numbers can have hex or binary prefixes (0x, 0b), or unsigned / length suffixes (ull), so we need
+	// to allow letters as a possibility, in addition to digits (other letters produce syntax errors, but it's fine if those
+	// return false and go through the slow path).
+	// 
+	// That leaves other operator characters we need to treat as ending the number:  ":<=>?^|~".  Looking at an ASCII chart shows
+	// a way to handle those.  All the above characters are at the bottom of the chart, and can be detected by looking at the
+	// bottom 5 bits of the character code:
+	// 
+	//    0x38 8  0x58 X  0x78  x
+	//    0x39 9  0x59 Y  0x79  y
+	//    0x3a :  0x5a Z  0x7a  z
+	//    0x3b ;  0x5b [  0x7b  {
+	//    0x3c <  0x5c \  0x7c  |
+	//    0x3d =  0x5d ]  0x7d  }
+	//    0x3e >  0x5e ^  0x7e  ~
+	//    0x3f ?  0x5f _  0x7f  DEL
+	//
+	// Anything 0x1a or larger in the bottom 5 bits is past the range of possible number continuations.  There's a handful of
+	// characters not caught by this logic:  "/@`", but two of those will produce syntax errors, and integer divide is rare in
+	// preprocessor expressions.  To avoid two branches, we merge two expressions which will have a negative sign for the
+	// condition we care about.  This is equivalent to:
+	//
+	//    (p < '.') || (0x1a <= (p & 0x1f))
+	//
+	return ((p - '.') | (0x19 - (p & 0x1f))) < 0;
+}
+
+inline int is_token_character(unsigned char p)
+{
+	// Numbers, letteres, underscore, and '$' are considered valid token characters by the preprocessor.  Mask test to see
+	// if this is one of those.
+	static const unsigned token_mask[8] = { 0x00000000, 0x03ff0010, 0x87fffffe, 0x07fffffe,  0,0,0,0 };
+
+	return (1u << (p & 0x1f)) & token_mask[p >> 5];
+}
+
 static void ce_next(ppcexp* c)
 {
 	char* p = c->p;
-	while (isspace(*p))
+	while (isspace_inline(*p))
 		++p;
+
 	switch (*p)
 	{
 		case '<':
@@ -200,7 +267,7 @@ static void ce_next(ppcexp* c)
 		{
 			do
 				++p;
-			while (isalnum(*p) || *p == '_' || *p == '$');
+			while (is_token_character(*p));
 			// any Identifier is a macro that WASN'T defined, so it evaluates as 0
 			c->token_value.iv = 0;
 			c->token_value.is_signed = 1;
@@ -209,6 +276,14 @@ static void ce_next(ppcexp* c)
 		}
 
 		case '0':
+			if (is_end_of_number(p[1]))
+			{
+				c->token_value.iv = 0;
+				c->token_value.is_signed = 1;
+				c->token = PPLEX_int_literal;
+				p++;
+				break;
+			}
 			if (p[1] == 'x' || p[1] == 'X')
 			{
 				char* q;
@@ -244,6 +319,14 @@ static void ce_next(ppcexp* c)
 		case '7':
 		case '8':
 		case '9':
+			if (is_end_of_number(p[1]))
+			{
+				c->token_value.iv = *p - '0';
+				c->token_value.is_signed = 1;
+				c->token = PPLEX_int_literal;
+				p++;
+				break;
+			}
 			c->token_value.uv = strtoull(p, &p, 10);
 		numeric_suffixes:
 			c->token_value.is_signed = (c->token_value.uv < ((uintmax_t)1 << 63));
diff --git a/Engine/Source/Developer/ShaderPreprocessor/Private/stb_preprocess/preprocessor.c b/Engine/Source/Developer/ShaderPreprocessor/Private/stb_preprocess/preprocessor.c
--- a/Engine/Source/Developer/ShaderPreprocessor/Private/stb_preprocess/preprocessor.c
+++ b/Engine/Source/Developer/ShaderPreprocessor/Private/stb_preprocess/preprocessor.c
@@ -173,13 +173,34 @@ typedef struct
 	shpair* verify;
 } pphash;
 
+#define MACRO_BLOOM_FILTER_ENABLED 1
+#define MACRO_BLOOM_FILTER_SIZE 128
+
 typedef struct pp_context
 {
 	pphash macro_map;				  // map of macro names to their definitions
 	macro_hash_entry* undef_map;	  // map of macro names to a flag for whether they've ever been undef'd, for more useful error messages
-	uint8 (*macro_name_filter)[256];  // 64KB table indicating whether 1 & 2-character identifiers are macros
 	include_once* once_map;			  // map of which filenames have been #onced
 
+#if MACRO_BLOOM_FILTER_ENABLED
+	// 3KB table by identifier length, with 3 masks indicating whether macros exist with the given key character at a certain location (first,
+	// middle, last).  Functions as a trivially cheap bloom filter with 4 keys (length plus the 3 characters), as it has fixed cost regardless
+	// of identifier length, and the filter table accesses for a given identifier are cache friendly, since character masks for a given length
+	// are adjacent.  Valid identifier characters are in the ASCII range 64-127, so we only need 64 bits of mask.
+	//
+	// In histogram testing, the middle character was actually the most unique, but all three make a contribution to improving effectiveness.
+	// The filter rejects 99.5% of identifiers that aren't macros in HLSL code, and handles 100% of the cases the original 1 or 2 character
+	// macro_name_filter covered in practice, meaning there's no reason to use both.
+	//
+	// The length of the identifier is clamped -- in testing, no identifiers longer than 64 characters were observed.  Using a fixed size
+	// table in the structure avoids a couple memory accesses fetching a pointer and current allocated length, which adds up given how
+	// often maybe_expand_macro is called.
+	uint64 macro_bloom_filter[MACRO_BLOOM_FILTER_SIZE][3];
+#else
+	// 64KB table indicating whether 1 & 2-character identifiers are macros
+	uint8(*macro_name_filter)[256];
+#endif
+
 	int include_nesting_level;	// we stop after a certain number in case of unbounded recursive includes
 	int macro_expansion_level;	// because recursive macros are prevented, we don't need to stop this, but we do in case of bugs in recursive macro processing
 	char* last_ifndef;
@@ -215,7 +236,6 @@ typedef struct parse_state
 	size_t src_length;
 
 	char* dest;
-	char* fast_dest;
 
 	int src_line_number;
 	int dest_line_number;
@@ -402,15 +422,10 @@ static void output_line_directive(parse_state* cs)
 	size_t filename_len = strlen(cs->filename);
 	size_t expansion_length = 6 + 16 + filename_len + 2 + 2 + 1;
 
-	if (cs->fast_dest)
-		q = cs->fast_dest;
-	else
-	{
-		// make sure there's room in the output for the #line directive plus all the remaining text
-		arrsetcap(out, arrlen(out) + (size_t)(cs->src_length - cs->src_offset) + expansion_length);
+	// make sure there's room in the output for the #line directive plus all the remaining text
+	arrsetcap(out, arrlen(out) + (size_t)(cs->src_length - cs->src_offset) + expansion_length);
 
-		q = out + arrlen(out);
-	}
+	q = out + arrlen(out);
 
 	// write the line directive
 	STB_ASSUME(q != NULL);
@@ -435,13 +450,9 @@ static void output_line_directive(parse_state* cs)
 	*q++ = '\n';
 
 	cs->dest_line_number = cs->src_line_number;
-	if (cs->fast_dest)
-		cs->fast_dest = q;
-	else
-	{
-		arrsetlen(out, (q - out));
-		cs->dest = out;
-	}
+
+	arrsetlen(out, (q - out));
+	cs->dest = out;
 }
 
 //////////////////////////////////////////////////////////////////////////////
@@ -1308,28 +1319,20 @@ static int copy_to_action_point(parse_state* cs)
 	int state;
 	const char* p = cs->src + cs->src_offset;
 	char* out;
-	char* q;
+	char* q = NULL;
 
 	out = cs->dest;
 
 	state = cs->in_leading_whitespace ? PP_STATE_in_leading_whitespace : PP_STATE_ready;
 	MDBG();
 
-	if (cs->fast_dest)
-		q = cs->fast_dest;
-	else
-		q = NULL;
-
 	for (;;)
 	{
 		int prev_state;
 
-		if (cs->fast_dest == NULL)
-		{
-			// expand output so there's room to copy entire file unmodified
-			arrsetcap(out, arrlen(out) + (cs->src_length - cs->src_offset) + 1);
-			q = out + arrlen(out);
-		}
+		// expand output so there's room to copy entire file unmodified
+		arrsetcap(out, arrlen(out) + (cs->src_length - cs->src_offset) + 1);
+		q = out + arrlen(out);
 
 		{  // make lifetime of line_count unambiguous
 			const int state_limit = cs->state_limit;
@@ -1339,7 +1342,7 @@ static int copy_to_action_point(parse_state* cs)
 			{
 				uint8 ch = (uint8)*p++;
 				uint8 ch_class;
-				assert(cs->fast_dest != NULL || q < out + arrcap(out));
+				assert(q < out + arrcap(out));
 				STB_ASSUME(q != NULL);
 				*q++ = ch;
 				ch_class = pp_char_class[ch];
@@ -1388,8 +1391,7 @@ static int copy_to_action_point(parse_state* cs)
 				cs->state_limit = PP_STATE_saw_cr;	// line_number mismatches, so need to output match newlines at first opportunity, namely next time we see a
 													// newline not in a comment
 				state = prev_state;					// reset parse state to before the backslash
-				if (cs->fast_dest == NULL)
-					arrsetlen(out, (size_t)(q - out));
+				arrsetlen(out, (size_t)(q - out));
 			}
 		}
 		else
@@ -1407,38 +1409,19 @@ static int copy_to_action_point(parse_state* cs)
 			cs->state_limit = PP_STATE_active_count;
 
 			// output newlines to the output buffer until synchronized again
-			if (cs->fast_dest)
-			{
-				int n = cs->src_line_number - cs->dest_line_number;
-				while (n > 0)
-				{
-					*q++ = '\n';
-					--n;
-				}
-				cs->dest_line_number = cs->src_line_number;
-			}
-			else
+			arrsetlen(out, (size_t)(q - out));
+			while (cs->dest_line_number < cs->src_line_number)
 			{
-				arrsetlen(out, (size_t)(q - out));
-				while (cs->dest_line_number < cs->src_line_number)
-				{
-					arrput(out, '\n');
-					++cs->dest_line_number;
-				}
+				arrput(out, '\n');
+				++cs->dest_line_number;
 			}
 		}
 	}
 
 	// set output to the amount we actually copied... don't count the copy of the last character
-	if (cs->fast_dest)
-	{
-		cs->fast_dest = q - 1;
-	}
-	else
-	{
-		arrsetlen(out, (size_t)(q - 1 - out));
-		cs->dest = out;
-	}
+	arrsetlen(out, (size_t)(q - 1 - out));
+	cs->dest = out;
+
 	cs->src_offset = (size_t)(p - 1 - cs->src);
 	cs->in_leading_whitespace = (state == PP_STATE_saw_leading_hash);
 
@@ -1461,14 +1444,9 @@ static int copy_to_action_point_macro_expansion(parse_state* cs)
 	STB_ASSUME(out != NULL);
 	char* q;
 
-	if (cs->fast_dest)
-		q = cs->fast_dest;
-	else
-	{
-		// expand output so copying entire file unmodified
-		arrsetcap(out, arrlen(out) + (cs->src_length - cs->src_offset) + 1);
-		q = out + arrlen(out);
-	}
+	// expand output so copying entire file unmodified
+	arrsetcap(out, arrlen(out) + (cs->src_length - cs->src_offset) + 1);
+	q = out + arrlen(out);
 
 	STB_ASSUME(q != NULL);
 
@@ -1485,14 +1463,9 @@ static int copy_to_action_point_macro_expansion(parse_state* cs)
 
 	cs->src_offset = (size_t)(p - 1 - cs->src);
 
-	if (cs->fast_dest)
-		cs->fast_dest = q - 1;
-	else
-	{
-		// set output to the amount we actually copied... don't count the copy of the last character
-		arrsetlen(out, (size_t)(q - 1 - out));
-		cs->dest = out;
-	}
+	// set output to the amount we actually copied... don't count the copy of the last character
+	arrsetlen(out, (size_t)(q - 1 - out));
+	cs->dest = out;
 
 	return state;
 }
@@ -2001,22 +1974,10 @@ static void preprocess_string(parse_state* cs, int in_macro_expansion, char* in_
 				if (cs->last_output_filename == cs->filename && cs->dest_line_number < cs->src_line_number && cs->src_line_number < cs->dest_line_number + 12)
 				{
 					int n = cs->src_line_number - cs->dest_line_number;
-					if (cs->fast_dest)
+					while (cs->dest_line_number < cs->src_line_number)
 					{
-						while (n > 0)
-						{
-							*cs->fast_dest++ = '\n';
-							--n;
-						}
-						cs->dest_line_number = cs->src_line_number;
-					}
-					else
-					{
-						while (cs->dest_line_number < cs->src_line_number)
-						{
-							arrput(cs->dest, '\n');
-							++cs->dest_line_number;
-						}
+						arrput(cs->dest, '\n');
+						++cs->dest_line_number;
 					}
 				}
 				else
@@ -2078,16 +2039,8 @@ static void preprocess_string(parse_state* cs, int in_macro_expansion, char* in_
 					// the final macro wasn't copied into the output destination yet, so do it now
 					{
 						size_t len = strlen(md->symbol_name);
-						if (cs->fast_dest)
-						{
-							memcpy(cs->fast_dest, md->symbol_name, len);
-							cs->fast_dest += len;
-						}
-						else
-						{
-							size_t off = arraddnindex(cs->dest, len);
-							memcpy(cs->dest + off, md->symbol_name, len);
-						}
+						size_t off = arraddnindex(cs->dest, len);
+						memcpy(cs->dest + off, md->symbol_name, len);
 					}
 				}
 				break;
@@ -2162,11 +2115,22 @@ static void preprocess_string(parse_state* cs, int in_macro_expansion, char* in_
 
 #define MAX_MACRO_PARAMETER_NAME_LENGTH 256
 
-static void update_macro_filter(pp_context* c, char* identifier)
+static void update_macro_filter(pp_context* c, char* identifier, size_t identifier_length)
 {
-	assert(identifier[0] != 0);
+	assert(identifier[0] != 0 && identifier_length > 0);
+
+#if MACRO_BLOOM_FILTER_ENABLED
+	size_t bloom_filter_offset = identifier_length < MACRO_BLOOM_FILTER_SIZE ? identifier_length : MACRO_BLOOM_FILTER_SIZE - 1;
+
+	// Set bits for three key characters of the macro by length (first, middle, last)
+	c->macro_bloom_filter[bloom_filter_offset][0] |= (1ull << (identifier[0] & 0x3f));
+	c->macro_bloom_filter[bloom_filter_offset][1] |= (1ull << (identifier[identifier_length >> 1] & 0x3f));
+	c->macro_bloom_filter[bloom_filter_offset][2] |= (1ull << (identifier[identifier_length - 1] & 0x3f));
+
+#else
 	if (identifier[1] == 0 || identifier[2] == 0)
 		c->macro_name_filter[(uint8)identifier[0]][(uint8)identifier[1]] = 1;
+#endif
 }
 
 static struct macro_definition* create_macro_definition(parse_state* ps, const char* def)
@@ -2607,7 +2571,7 @@ static const char* copy_argument(const char* text, int* line_number, char** p_ou
 		{
 			++p;
 		}
-		oldlen = arrlen(out);
+		oldlen = arrlennonull(out);
 		addlen = p - q;
 		STB_ASSUME(out != NULL && oldlen > 0);
 		arrsetlen(out, oldlen + addlen);
@@ -2734,12 +2698,10 @@ static void maybe_expand_macro(parse_state* cs, struct macro_definition* pending
 {
 	pp_context* c = cs->context;
 	struct macro_definition* md;
-	char* identifier = cs->fast_dest ? cs->fast_dest : (char*)cs->dest + arrlennonull(cs->dest);
+	char* identifier = (char*)cs->dest + arrlennonull(cs->dest);
 	size_t identifier_length;
 	const char* p = cs->src + cs->src_offset;
 	char* q = identifier;
-	uint8 c0 = (uint8)p[0];
-	uint8 c1 = (uint8)p[1];
 
 	if (pending)
 	{
@@ -2750,7 +2712,10 @@ static void maybe_expand_macro(parse_state* cs, struct macro_definition* pending
 	{
 		const char* start = p;
 
-#ifndef DISABLE_SHORT_MACRO_FILTER
+#if !MACRO_BLOOM_FILTER_ENABLED
+		uint8 c0 = (uint8)p[0];
+		uint8 c1 = (uint8)p[1];
+
 		if (!char_is_pp_identifier(c1))
 		{
 			// 1-character identifier
@@ -2778,19 +2743,31 @@ static void maybe_expand_macro(parse_state* cs, struct macro_definition* pending
 		// copy the rest of the identifier so we can NUL-terminate it
 		// we copy it to the output buffer since we've already allocated space for it,
 		// and that's where we'll need to copy it if it's not a macro
+		// There is always at least 1 identifier character, so copy the first unconditionally
+		char ch_first = *p++;
+		*q++ = ch_first;
 		while (char_is_pp_identifier(*p))
 			*q++ = *p++;
 
 		identifier_length = p - start;
+
+#if MACRO_BLOOM_FILTER_ENABLED
+		// Test bits for three key characters of the macro by length (first, middle, last)
+		size_t bloom_filter_offset = identifier_length < MACRO_BLOOM_FILTER_SIZE ? identifier_length : MACRO_BLOOM_FILTER_SIZE - 1;
+		if (!(c->macro_bloom_filter[bloom_filter_offset][0] & (1ull << ch_first)) ||
+			!(c->macro_bloom_filter[bloom_filter_offset][1] & (1ull << start[identifier_length >> 1])) ||
+			!(c->macro_bloom_filter[bloom_filter_offset][2] & (1ull << start[identifier_length - 1])))
+		{
+			goto not_macro;
+		}
+#endif
+
 		md = (struct macro_definition*)stringhash_get(&c->macro_map, identifier, identifier_length);
 		if (md == NULL || md->disabled)
 		{
 		not_macro:
 			cs->src_offset = p - cs->src;
-			if (cs->fast_dest)
-				cs->fast_dest = q;
-			else
-				arrsetlen(cs->dest, (size_t)(q - cs->dest));  // set the output buffer size to account for the above-copied identifier
+			arrsetlen(cs->dest, (size_t)(q - cs->dest));  // set the output buffer size to account for the above-copied identifier
 			return;
 		}
 
@@ -2813,10 +2790,7 @@ static void maybe_expand_macro(parse_state* cs, struct macro_definition* pending
 					//   command:            cpp -DFOO foo.c
 					//   desired output:     + +x;
 					//   without next line:  ++x;
-					if (cs->fast_dest)
-						*cs->fast_dest++ = ' ';
-					else
-						arrput(cs->dest, ' ');
+					arrput(cs->dest, ' ');
 				else
 				{
 					parse_state ncs = *cs;
@@ -2828,7 +2802,6 @@ static void maybe_expand_macro(parse_state* cs, struct macro_definition* pending
 					preprocess_string(&ncs, IN_MACRO_yes, md->symbol_name, final_function_macro);
 					md->disabled = 0;
 					cs->dest = ncs.dest;
-					cs->fast_dest = ncs.fast_dest;
 				}
 				break;
 			}
@@ -2891,33 +2864,19 @@ static void maybe_expand_macro(parse_state* cs, struct macro_definition* pending
 			{
 				size_t len = strlen(cs->filename);
 				// @TODO: backslash escape filename and cache it in stack, ready for string printing... also use when printing the line number?
-				if (cs->fast_dest)
-				{
-					*cs->fast_dest++ = '"';
-					cs->fast_dest += sprintf(cs->fast_dest, "%s", cs->filename);
-					*cs->fast_dest++ = '"';
-				}
-				else
-				{
-					size_t off = arraddnindex(cs->dest, len + 2);
-					cs->dest[off] = '"';
-					memcpy(cs->dest + off + 1, cs->filename, len);
-					cs->dest[off + 1 + len] = '"';
-				}
+				size_t off = arraddnindex(cs->dest, len + 2);
+				cs->dest[off] = '"';
+				memcpy(cs->dest + off + 1, cs->filename, len);
+				cs->dest[off + 1 + len] = '"';
 				break;
 			}
 			case MACRO_NUM_PARAMETERS_line:
 			{
-				if (cs->fast_dest)
-					cs->fast_dest += sprintf(cs->fast_dest, "%d", cs->src_line_number);
-				else
-				{
-					int len;
-					q = identifier;
-					arrsetcap(cs->dest, arrlen(cs->dest) + 12);
-					len = sprintf(cs->dest + arrlen(cs->dest), "%d", cs->src_line_number);
-					arrsetlen(cs->dest, arrlen(cs->dest) + len);
-				}
+				int len;
+				q = identifier;
+				arrsetcap(cs->dest, arrlen(cs->dest) + 12);
+				len = sprintf(cs->dest + arrlen(cs->dest), "%d", cs->src_line_number);
+				arrsetlen(cs->dest, arrlen(cs->dest) + len);
 				break;
 			}
 			case MACRO_NUM_PARAMETERS_counter:
@@ -2928,16 +2887,11 @@ static void maybe_expand_macro(parse_state* cs, struct macro_definition* pending
 					do_error_code(cs, PP_RESULT_counter_overflowed, "__COUNTER__ was used too many times (%u)", c->counter);
 					return;
 				}
-				if (cs->fast_dest)
-					cs->fast_dest += sprintf(cs->fast_dest, "%d", c->counter);
-				else
-				{
-					int len;
-					q = identifier;
-					arrsetcap(cs->dest, arrlen(cs->dest) + 12);
-					len = sprintf(cs->dest + arrlen(cs->dest), "%d", c->counter);
-					arrsetlen(cs->dest, arrlen(cs->dest) + len);
-				}
+				int len;
+				q = identifier;
+				arrsetcap(cs->dest, arrlen(cs->dest) + 12);
+				len = sprintf(cs->dest + arrlen(cs->dest), "%d", c->counter);
+				arrsetlen(cs->dest, arrlen(cs->dest) + len);
 				++c->counter;
 				break;
 			}
@@ -2990,12 +2944,9 @@ static void maybe_expand_macro(parse_state* cs, struct macro_definition* pending
 		p = s + 1;
 
 #if 0
-    // WHAT?!?
-    // undo '(' from output
-    if (cs->fast_dest)
-      --cs->fast_dest;
-    else
-      arrsetlen(cs->dest, arrlen(cs->dest)-1);
+	// WHAT?!?
+	// undo '(' from output
+	arrsetlen(cs->dest, arrlen(cs->dest)-1);
 #endif
 
 		//////////////////////////////////////////////////////////////////
@@ -3072,7 +3023,7 @@ static void maybe_expand_macro(parse_state* cs, struct macro_definition* pending
 				// the call to arrsetlen below, because downstream code assumes there is a null terminator beyond the official length of the array.  Doing
 				// the sub-allocation first preserves the null terminator -- without this, the sub-allocation will overwrite it.  We only suballocate if
 				// a minimum of 200 characters is available, as copy_argument reserves this.
-				stbds_arrinline_suballoc(argument_buffer, char, 200);
+				stbds_arrinline_suballoc_char(argument_buffer, 200);
 
 				arrsetlen(copy, e - p);
 				arrput(arguments, copy);
@@ -3149,8 +3100,6 @@ static void maybe_expand_macro(parse_state* cs, struct macro_definition* pending
 					e = copy;
 				*e = 0;
 
-				stbds_arrinline_suballoc(argument_buffer, char, 200);
-
 				arrput(arguments, copy);
 				cs->src_line_number += arg_newlines;
 			}
@@ -3177,7 +3126,7 @@ static void maybe_expand_macro(parse_state* cs, struct macro_definition* pending
 
 		// each argument will be parsed using acs, so set that up now
 		acs = *cs;
-		acs.dest = acs.fast_dest = 0;
+		acs.dest = 0;
 
 		// populate temporary buffer
 		char* tmp = 0;
@@ -3243,7 +3192,6 @@ static void maybe_expand_macro(parse_state* cs, struct macro_definition* pending
 				acs.src = arguments[an];
 				acs.src_length = arrlen(arguments[an]);
 				acs.dest = tmp;
-				acs.fast_dest = 0;
 				acs.src_offset = 0;
 				preprocess_string(&acs, in_macro_expansion, identifier, &ffm);
 				tmp = acs.dest;
@@ -3284,7 +3232,6 @@ static void maybe_expand_macro(parse_state* cs, struct macro_definition* pending
 		md->disabled = 1;  // don't disable earlier, because if you CALL foo() WITH foo, it should still get expanded
 		preprocess_string(&ncs, IN_MACRO_yes, identifier, final_function_macro);
 		cs->dest = ncs.dest;
-		cs->fast_dest = ncs.fast_dest;
 		arrfree(ncs.src);
 
 		md->disabled = 0;
@@ -3305,7 +3252,6 @@ static char* macro_expand_directive(parse_state* cs, const char* p, char* direct
 	ncs.src_offset = 0;
 	ncs.parent = cs;
 	ncs.dest = 0;
-	ncs.fast_dest = 0;
 	preprocess_string(&ncs, IN_MACRO_directive, directive, NULL);
 	return ncs.dest;
 }
@@ -3323,7 +3269,6 @@ static int evaluate_if(parse_state* cs, const char* p, int* syntax_error)
 	ncs.src_offset = 0;
 	ncs.parent = cs;
 	ncs.dest = dest;
-	ncs.fast_dest = 0;
 
 	preprocess_string(&ncs, IN_MACRO_if_condition, "#if", NULL);
 
@@ -3471,10 +3416,7 @@ static int process_include(parse_state* cs, const char* start, conditional_state
 				
 				(*freefile_callback)(filename, data, c->custom_context);
 
-				if (cs->fast_dest)
-					cs->fast_dest = result;
-				else
-					cs->dest = result;
+				cs->dest = result;
 				if (c->stop)
 					return RETURN_BEHAVIOR_return;
 			}
@@ -3604,17 +3546,6 @@ static void process_directive(parse_state* cs, conditional_state* cons)
 
 		case HASH_line:
 			// copy line directive through to output... existing code will output the newline, so don't do it here
-			if (cs->fast_dest)
-			{
-				size_t out_length = strlen(p);
-				char* z = cs->fast_dest;
-				memcpy(z, "#line ", 6);
-				z += 6;
-				memcpy(z, p, out_length);
-				z += out_length;
-				cs->fast_dest = z;
-			}
-			else
 			{
 				size_t out_length = strlen(p);
 				size_t s = arraddnindex(cs->dest, 6);
@@ -3713,17 +3644,6 @@ static void process_directive(parse_state* cs, conditional_state* cons)
 			}
 
 			// copy pragma through to output... existing code will output the newline, so don't do it here
-			if (cs->fast_dest)
-			{
-				size_t out_length = strlen(p);
-				char* z = cs->fast_dest;
-				memcpy(z, "#pragma ", 8);
-				z += 8;
-				memcpy(z, p, out_length);
-				z += out_length;
-				cs->fast_dest = z;
-			}
-			else
 			{
 				size_t out_length = strlen(p);
 				size_t s = arraddnindex(cs->dest, 8);
@@ -3757,8 +3677,7 @@ static void process_directive(parse_state* cs, conditional_state* cons)
 					stringhash_delete(&c->macro_map, m->symbol_name, m->symbol_name_length);
 				}
 				stringhash_put(&c->macro_map, m->symbol_name, m->symbol_name_length, m);
-				if (m->symbol_name[1] == 0 || m->symbol_name[2] == 0)
-					c->macro_name_filter[(uint8)m->symbol_name[0]][(uint8)m->symbol_name[1]] = 1;
+				update_macro_filter(c, m->symbol_name, m->symbol_name_length);
 			}
 			if (alloc)
 				arrfree(alloc);
@@ -4287,14 +4206,11 @@ static char* preprocess_string_from_file(parse_state* ps, const char* filename,
 	cs.src_length = textlen;
 	cs.src_offset = 0;
 	cs.conditional_nesting_depth_at_start = conditional_compilation_nesting;
-	// dest & fast_dest are inherited
+	// dest is inherited
 	preprocess_string(&cs, is_header && preprocessor_automatic_include_guard_detection ? IN_MACRO_include_guard_scan : IN_MACRO_no, 0, 0);
 	cs.context->num_lines += cs.src_line_number;
 	ps->last_output_filename = "";
-	if (cs.fast_dest)
-		return cs.fast_dest;
-	else
-		return cs.dest;
+	return cs.dest;
 }
 
 #define DEFAULT -1
@@ -4588,8 +4504,7 @@ static void define_macro(pphash* map, struct macro_definition* m)
 	stringhash_put(map, m->symbol_name, m->symbol_name_length, m);
 }
 
-char* preprocess_file(char* output_autobuffer,
-	const char* filename,
+char* preprocess_file(const char* filename,
 	void* custom_context,
 	struct macro_definition** predefined_macros,
 	int num_predefined_macros,
@@ -4631,11 +4546,15 @@ char* preprocess_file(char* output_autobuffer,
 	// backfire by slowing down the initialization
 	stringhash_create(&c.macro_map, 4096);
 
+#if MACRO_BLOOM_FILTER_ENABLED
+	memset(c.macro_bloom_filter, 0, sizeof(c.macro_bloom_filter));
+#else
 	c.macro_name_filter = STB_COMMON_MALLOC(65536);
 	if (c.macro_name_filter == NULL)
 		return NULL;
 
 	memset(c.macro_name_filter, 0, 65536);
+#endif
 
 	for (i = 0; i < arrlen(predefined_macros); ++i)
 		define_macro(&c.macro_map, predefined_macros[i]);
@@ -4646,26 +4565,17 @@ char* preprocess_file(char* output_autobuffer,
 	define_macro(&c.macro_map, &predefined_COUNTER);
 
 	for (i = 0; i < arrlen(c.macro_map.pair); ++i)
-		update_macro_filter(&c, c.macro_map.pair[i].key);
+		update_macro_filter(&c, c.macro_map.pair[i].key, c.macro_map.pair[i].key_length);
 
 	{
 		// initial parse_state for preprocess_string_from_file to initialize from
 		parse_state ps = {0};
 		ps.context = &c;
 		ps.dest = output;
-		ps.fast_dest = output_autobuffer;
 		ps.state_limit = PP_STATE_active_count;
 
 		output = preprocess_string_from_file(&ps, filename, main_file, length, 0, 0);
-		if (output_autobuffer)
-		{
-			*output = 0;
-			output = output_autobuffer;
-		}
-		else
-		{
-			arrpush(output, 0);
-		}
+		arrpush(output, 0);
 	}
 
 	(*freefile_callback)(filename, main_file, custom_context);
@@ -4686,7 +4596,9 @@ char* preprocess_file(char* output_autobuffer,
 	shfree(c.undef_map);
 	arrfree(c.ifdef_stack);
 	stringhash_destroy(&c.macro_map);
+#if !MACRO_BLOOM_FILTER_ENABLED
 	STB_COMMON_FREE(c.macro_name_filter);
+#endif
 	stb_arena_free(&c.macro_arena);
 
 	if (c.diagnostics)
@@ -4695,8 +4607,7 @@ char* preprocess_file(char* output_autobuffer,
 		*num_pd = (int)arrlen(*pd);
 		if (c.stop)
 		{
-			if (output_autobuffer == NULL)
-				arrfree(output);
+			arrfree(output);
 			output = NULL;
 		}
 	}
diff --git a/Engine/Source/Developer/ShaderPreprocessor/Private/stb_preprocess/preprocessor.h b/Engine/Source/Developer/ShaderPreprocessor/Private/stb_preprocess/preprocessor.h
--- a/Engine/Source/Developer/ShaderPreprocessor/Private/stb_preprocess/preprocessor.h
+++ b/Engine/Source/Developer/ShaderPreprocessor/Private/stb_preprocess/preprocessor.h
@@ -85,7 +85,6 @@ STB_PP_DEF void init_preprocessor(
 // in the pd array (setting num_pd to the diagnostic count). preprocessor_file_free should be called for 
 // each invocation to free any allocated memory.
 STB_PP_DEF char* preprocess_file(
-	char* output_storage,
 	const char* filename,
 	void* custom_context,
 	struct macro_definition** predefined_macros,
diff --git a/Engine/Source/Developer/ShaderPreprocessor/Private/stb_preprocess/stb_ds.h b/Engine/Source/Developer/ShaderPreprocessor/Private/stb_preprocess/stb_ds.h
--- a/Engine/Source/Developer/ShaderPreprocessor/Private/stb_preprocess/stb_ds.h
+++ b/Engine/Source/Developer/ShaderPreprocessor/Private/stb_preprocess/stb_ds.h
@@ -516,16 +516,47 @@ extern "C"
 #define stbds_temp(t) stbds_header(t)->temp
 #define stbds_temp_key(t) (*(char**)stbds_header(t)->hash_table)
 
-// Macro to reserve inline (usually stack memory) storage for an array
+// Macro to reserve inline stack memory storage for an array
 #define stbds_arrinline(a, a_type, n) \
-	size_t __buf ## a[(sizeof(stbds_array_header) + n * sizeof(*a) + sizeof(size_t) - 1) / sizeof(size_t)];	\
-	a = (a_type*)stbds_arrinlinef(__buf ## a, sizeof(*a), n)
+	size_t __buf ## a[(sizeof(stbds_array_header) + n * sizeof(a_type) + sizeof(size_t) - 1) / sizeof(size_t)];		\
+	((stbds_array_header*)(&__buf ## a[0]))->length = 0;															\
+	((stbds_array_header*)(&__buf ## a[0]))->capacity = n;															\
+	((stbds_array_header*)(&__buf ## a[0]))->hash_table = 0;														\
+	((stbds_array_header*)(&__buf ## a[0]))->temp = 0;																\
+	((stbds_array_header*)(&__buf ## a[0]))->elemsize = sizeof(a_type);												\
+	((stbds_array_header*)(&__buf ## a[0]))->inlinealloc = 1;														\
+	a = (a_type*)(((stbds_array_header*)(&__buf ## a[0])) + 1)
 
 // Allow a single inline storage buffer to be used for multiple sub-allocations.  Trims current allocation to what's used and returns a new allocation with the remainder if enough capacity exists.
 #define stbds_arrinline_suballoc(a, a_type, min_capacity) \
 	a = (a_type*)stbds_arrinline_suballocf(a, min_capacity)
 
-#define stbds_arrsetcap(a, n) (stbds_arrgrow(a, 0, n))
+// Inline simplified version of above that only works on character type.  Avoids elemsize multiply/divide math and call overhead.
+#define stbds_arrinline_suballoc_char(a, min_capacity) \
+	if (a && stbds_header(a)->inlinealloc)														\
+	{																							\
+		stbds_array_header* h = stbds_header(a);												\
+		char* used_end = a + ((h->length + sizeof(size_t) - 1) & ~(sizeof(size_t) - 1));		\
+		char* capacity_end = a + h->capacity;													\
+		h->capacity = h->length;																\
+		if (capacity_end - used_end < sizeof(stbds_array_header) + min_capacity)				\
+		{																						\
+			a = 0;																				\
+		}																						\
+		else																					\
+		{																						\
+			stbds_array_header* h_suballoc = (stbds_array_header*)used_end;						\
+			h_suballoc->length = 0;																\
+			h_suballoc->capacity = ((capacity_end - used_end) - sizeof(stbds_array_header));	\
+			h_suballoc->hash_table = 0;															\
+			h_suballoc->temp = 0;																\
+			h_suballoc->elemsize = 1;															\
+			h_suballoc->inlinealloc = 1;														\
+			a = (char*)(h_suballoc + 1);														\
+		}}   (void)0
+
+
+#define stbds_arrsetcap(a, n) ((!(a) || (n) > stbds_header(a)->capacity) ? (stbds_arrgrow(a, 0, n), 0) : 0)
 #define stbds_arrsetlen(a, n) ((stbds_arrcap(a) < (size_t)(n) ? stbds_arrsetcap((a), (size_t)(n)), 0 : 0), (a) ? stbds_header(a)->length = (size_t)(n) : 0)
 #define stbds_arrinitlen(a, n) ((stbds_arrsetcap((a), (size_t)(n)), stbds_header(a)->length = (size_t)(n)))
 #define stbds_arrcap(a) ((a) ? stbds_header(a)->capacity : 0)
@@ -628,12 +659,12 @@ extern "C"
 
 typedef struct
 {
-	size_t length;
-	size_t capacity;
 	int elemsize;
 	int inlinealloc;		// Allocated inline, don't free!
 	void* hash_table;
 	ptrdiff_t temp;
+	int capacity;
+	int length;
 } stbds_array_header;
 
 typedef struct stbds_string_block
diff --git a/Engine/Source/Runtime/RenderCore/Internal/ShaderCompilerDefinitions.h b/Engine/Source/Runtime/RenderCore/Internal/ShaderCompilerDefinitions.h
--- a/Engine/Source/Runtime/RenderCore/Internal/ShaderCompilerDefinitions.h
+++ b/Engine/Source/Runtime/RenderCore/Internal/ShaderCompilerDefinitions.h
@@ -145,16 +145,16 @@ public:
 			switch (Other.ValueTypes[OtherIndex])
 			{
 			case EShaderCompilerDefineVariant::Integer:
-				SetDefine(OtherIt.Key(), Other.Pairs[OtherIndex].ValueInteger);
+				SetDefine(OtherIt.KeyFName(), Other.Pairs[OtherIndex].ValueInteger);
 				break;
 			case EShaderCompilerDefineVariant::Unsigned:
-				SetDefine(OtherIt.Key(), Other.Pairs[OtherIndex].ValueUnsigned);
+				SetDefine(OtherIt.KeyFName(), Other.Pairs[OtherIndex].ValueUnsigned);
 				break;
 			case EShaderCompilerDefineVariant::Float:
-				SetDefine(OtherIt.Key(), Other.Pairs[OtherIndex].ValueFloat);
+				SetDefine(OtherIt.KeyFName(), Other.Pairs[OtherIndex].ValueFloat);
 				break;
 			case EShaderCompilerDefineVariant::String:
-				SetDefine(OtherIt.Key(), OtherIt.Value());
+				SetDefine(OtherIt.KeyFName(), OtherIt.Value());
 				break;
 			}
 		}
@@ -246,6 +246,11 @@ public:
 			return KeyStringBuffer;
 		}
 
+		FORCEINLINE const FName& KeyFName() const
+		{
+			return Defines.Pairs[Index].Key;
+		}
+
 		// Note that the output of Value() is transient, only valid until the iterator is incremented!
 		FORCEINLINE const TCHAR* Value()
 		{
