Shader Preprocessor:  Early bloom filter and SSE optimizations.  Overall 27.9% improvement to low level preprocessor, or 10.4% to ConditionalPreprocessShader as a whole.

* Moved identifier copy and macro bloom filter from maybe_expand_macro to test into copy_to_action_point / copy_to_action_point_macro_expansion.  13.1% of improvement.
* SSE implementation of scan_to_directive, 10x faster, 5.2%
* SSE implementation of identifier copy, 3x faster, 4.5%
* SSE ShaderConvertAndStripComments, 4x faster, 3.6%
* Fast inline string equality comparison, 4x faster, 1.5%

To make SSE implementations "safe" without needing special cases near the end of a buffer, it's necessary to ensure padding is present in the relevant buffers, anything that goes through a preprocess_string call.  This includes the string arena allocator, temporary stbds arrays that hold strings, and file buffers passed in.  The latter all pass through ShaderConvertAndStripComments, where we can add padding.  (ShaderConvertAndStripComments itself has special cases for end of buffer).  Code related to original 1 and 2 character macro filter removed, since I can't see a reason to enable it over the bloom filter.

I also attempted SSE optimization of copy_to_action_point and copy_line_without_comments, but improvement wasn't big enough to be worth the complexity (around 2% for the former, but massive code complexity, 0.5% for the latter).  That's pretty much everything SSE friendly that's over 1% on a profile, although I think copy_argument can be made a lot faster, not primarily through SSE.

#rnx

diff --git a/Engine/Source/Developer/ShaderPreprocessor/Private/ShaderPreprocessor.cpp b/Engine/Source/Developer/ShaderPreprocessor/Private/ShaderPreprocessor.cpp
--- a/Engine/Source/Developer/ShaderPreprocessor/Private/ShaderPreprocessor.cpp
+++ b/Engine/Source/Developer/ShaderPreprocessor/Private/ShaderPreprocessor.cpp
@@ -103,6 +103,7 @@ struct FStbLoadedInclude
 {
 	const ANSICHAR* Data = nullptr;				// Points to SharedData, LocalData, or data from FShaderCompilerEnvironment
 	size_t DataLength = 0;
+	size_t DataCapacity = 0;
 	FShaderSharedAnsiStringPtr SharedData;
 	TArray<ANSICHAR> LocalData;
 };
@@ -139,6 +140,7 @@ static const ANSICHAR* StbLoadFile(const ANSICHAR* Filename, void* RawContext, s
 
 			ContentsCached->Data = ContentsCached->LocalData.GetData();
 			ContentsCached->DataLength = ContentsCached->LocalData.Num();
+			ContentsCached->DataCapacity = ContentsCached->LocalData.Max();
 		}
 		else
 		{
@@ -148,6 +150,7 @@ static const ANSICHAR* StbLoadFile(const ANSICHAR* Filename, void* RawContext, s
 			{
 				ContentsCached->Data = InMemorySourceAnsi->Get()->GetData();
 				ContentsCached->DataLength = InMemorySourceAnsi->Get()->Num();
+				ContentsCached->DataCapacity = InMemorySourceAnsi->Get()->Max();
 			}
 			else
 			{
@@ -156,8 +159,14 @@ static const ANSICHAR* StbLoadFile(const ANSICHAR* Filename, void* RawContext, s
 
 				ContentsCached->Data = ContentsCached->SharedData->GetData();
 				ContentsCached->DataLength = ContentsCached->SharedData->Num();
+				ContentsCached->DataCapacity = ContentsCached->SharedData->Max();
 			}
 		}
+
+		// Need 15 characters beyond null terminator, so an unaligned SSE read at the null terminator can safely read 15 extra unused characters
+		// without going out of memory bounds.  ShaderConvertAndStripComments ensures this padding.  We could optionally allocate (or reallocate)
+		// a local copy as a fallback to handle this case without asserting, but it would be a silent performance degradation.
+		checkf(ContentsCached->DataCapacity >= ContentsCached->DataLength + 15, TEXT("Shader preprocessor ANSI files must include 15 bytes of capacity padding past null terminator"));
 	}
 	check(ContentsCached);
 	*OutLength = ContentsCached->DataLength;
diff --git a/Engine/Source/Developer/ShaderPreprocessor/Private/stb_preprocess/preprocessor.c b/Engine/Source/Developer/ShaderPreprocessor/Private/stb_preprocess/preprocessor.c
--- a/Engine/Source/Developer/ShaderPreprocessor/Private/stb_preprocess/preprocessor.c
+++ b/Engine/Source/Developer/ShaderPreprocessor/Private/stb_preprocess/preprocessor.c
@@ -26,6 +26,16 @@
 #define MDBG()
 #endif
 
+// Unreal 5.2+ assumes SSE 4.2 is available on x64 processors
+#if (defined(_M_X64) || defined(__amd64__) || defined(__x86_64__)) && !defined(_M_ARM64EC)
+#include <immintrin.h>
+#define PREPROCESSOR_USE_SSE4_2 1
+#define SSE_READ_PADDING 16
+#else
+#define PREPROCESSOR_USE_SSE4_2 0
+#define SSE_READ_PADDING 0
+#endif
+
 #pragma warning(push)
 // arrput calls are incorrectly reported by MSVC analysis to be potentially dereferencing null
 // (the maybegrow macro will inline allocate if given a null pointer to begin with)
@@ -106,6 +116,13 @@ static struct macro_definition predefined_defined = {"defined", 7, {0}, 0, MACRO
 #define strdup(x) _strdup(x)
 #endif
 
+#if !defined(FORCE_INLINE)
+#if defined(_MSC_VER)
+#define FORCE_INLINE __forceinline
+#else
+#define FORCE_INLINE inline __attribute__((always_inline))
+#endif
+#endif
 
 typedef struct
 {
@@ -117,7 +134,7 @@ typedef struct
 #define DIRECTIVE_HASH_SIZE 32
 struct
 {
-	char* name;
+	char name[8];
 	size_t name_len;
 } directive_hash[DIRECTIVE_HASH_SIZE];
 
@@ -173,7 +190,6 @@ typedef struct
 	shpair* verify;
 } pphash;
 
-#define MACRO_BLOOM_FILTER_ENABLED 1
 #define MACRO_BLOOM_FILTER_SIZE 128
 
 typedef struct pp_context
@@ -182,24 +198,18 @@ typedef struct pp_context
 	macro_hash_entry* undef_map;	  // map of macro names to a flag for whether they've ever been undef'd, for more useful error messages
 	include_once* once_map;			  // map of which filenames have been #onced
 
-#if MACRO_BLOOM_FILTER_ENABLED
 	// 3KB table by identifier length, with 3 masks indicating whether macros exist with the given key character at a certain location (first,
 	// middle, last).  Functions as a trivially cheap bloom filter with 4 keys (length plus the 3 characters), as it has fixed cost regardless
 	// of identifier length, and the filter table accesses for a given identifier are cache friendly, since character masks for a given length
 	// are adjacent.  Valid identifier characters are in the ASCII range 64-127, so we only need 64 bits of mask.
 	//
 	// In histogram testing, the middle character was actually the most unique, but all three make a contribution to improving effectiveness.
-	// The filter rejects 99.5% of identifiers that aren't macros in HLSL code, and handles 100% of the cases the original 1 or 2 character
-	// macro_name_filter covered in practice, meaning there's no reason to use both.
+	// The filter rejects 99.5% of identifiers that aren't macros in HLSL code.
 	//
 	// The length of the identifier is clamped -- in testing, no identifiers longer than 64 characters were observed.  Using a fixed size
 	// table in the structure avoids a couple memory accesses fetching a pointer and current allocated length, which adds up given how
 	// often maybe_expand_macro is called.
 	uint64 macro_bloom_filter[MACRO_BLOOM_FILTER_SIZE][3];
-#else
-	// 64KB table indicating whether 1 & 2-character identifiers are macros
-	uint8(*macro_name_filter)[256];
-#endif
 
 	int include_nesting_level;	// we stop after a certain number in case of unbounded recursive includes
 	int macro_expansion_level;	// because recursive macros are prevented, we don't need to stop this, but we do in case of bugs in recursive macro processing
@@ -236,6 +246,7 @@ typedef struct parse_state
 	size_t src_length;
 
 	char* dest;
+	size_t copied_identifier_length;			// Set by copy_and_filter_macro, used by maybe_expand_macro
 
 	int src_line_number;
 	int dest_line_number;
@@ -713,6 +724,44 @@ void stringhash_init_table(pphash* ht, unsigned int size)
 		table[i] = t;  // hope for 8-byte-at-a-time init
 }
 
+static int fast_preprocess_string_equals(const char* a, const char* b, int length)
+{
+#if PREPROCESSOR_USE_SSE4_2
+	// Does equality compare rather than ordinal compare, and assumes it's safe to read 7 bytes past end of buffer,
+	// which is true for strings in the preprocessor, which are in padded arenas, source files, or buffers.
+	if (length > 8)
+	{
+		// Longer string, compare 8 byte words, then a pair of possibly overlapping words at the end.
+		for (; length > 16; length-=8, a+=8, b+=8)
+		{
+			if (*(const uint64*)a != *(const uint64*)b)
+			{
+				return 0;
+			}
+		}
+		return ((*(const uint64*)a ^ *(const uint64*)b) | (*(const uint64*)(a + length - 8) ^ *(const uint64*)(b + length - 8))) == 0;
+	}
+	else
+	{
+		// Short string, compare up to 8 masked bytes.  May read up to 7 bytes past end of input strings.
+		return (((uint64)(-1ll)) >> ((8 - length) * 8) & (*(const uint64*)a ^ *(const uint64*)b)) == 0;
+	}
+#else
+	return memcmp(a, b, length) == 0;
+#endif
+}
+
+// For known short length strings (such as directives)
+static int fast_preprocess_string_equals_short(const char* a, const char* b, int length)
+{
+#if PREPROCESSOR_USE_SSE4_2
+	// Short string, compare up to 8 masked bytes.  May read up to 7 bytes past end of input strings.
+	return (((uint64)(-1ll)) >> ((8 - length) * 8) & (*(const uint64*)a ^ *(const uint64*)b)) == 0;
+#else
+	return memcmp(a, b, length) == 0;
+#endif
+}
+
 void* stringhash_get(pphash* ht, const char* key, size_t keylen)
 {
 	uint32 mask = ht->table_mask;
@@ -735,7 +784,7 @@ void* stringhash_get(pphash* ht, const char* key, size_t keylen)
 	if (table[slot].hash == hash)
 	{
 		uint32 i = table[slot].index;
-		if (ht->pair[i].key_length == keylen && 0 == memcmp(key, ht->pair[i].key, keylen))
+		if (ht->pair[i].key_length == keylen && fast_preprocess_string_equals(key, ht->pair[i].key, keylen))
 		{
 			if (ht->verify)
 				assert(vresult == ht->pair[i].value);
@@ -756,7 +805,7 @@ void* stringhash_get(pphash* ht, const char* key, size_t keylen)
 		if (table[slot].hash == hash)
 		{
 			uint32 i = table[slot].index;
-			if (ht->pair[i].key_length == keylen && 0 == memcmp(key, ht->pair[i].key, keylen))
+			if (ht->pair[i].key_length == keylen && fast_preprocess_string_equals(key, ht->pair[i].key, keylen))
 			{
 				if (ht->verify)
 					assert(vresult == ht->pair[i].value);
@@ -792,7 +841,7 @@ int stringhash_delete(pphash* ht, char* key, size_t keylen)
 	if (table[slot].hash == hash)
 	{
 		i = table[slot].index;
-		if (ht->pair[i].key_length == keylen && 0 == memcmp(key, ht->pair[i].key, keylen))
+		if (ht->pair[i].key_length == keylen && fast_preprocess_string_equals(key, ht->pair[i].key, keylen))
 			goto del;
 	}
 	if (table[slot].hash == HASH_EMPTY_MARKER)
@@ -806,7 +855,7 @@ int stringhash_delete(pphash* ht, char* key, size_t keylen)
 		if (table[slot].hash == hash)
 		{
 			i = table[slot].index;
-			if (ht->pair[i].key_length == keylen && 0 == memcmp(key, ht->pair[i].key, keylen))
+			if (ht->pair[i].key_length == keylen && fast_preprocess_string_equals(key, ht->pair[i].key, keylen))
 				goto del;
 		}
 		if (table[slot].hash == HASH_EMPTY_MARKER)
@@ -1308,6 +1357,73 @@ static int scan_whitespace_and_comments(parse_state* cs)
 	return state;
 }
 
+// Returns NULL if possibly a macro, or the end of the identifier if definitely not a macro
+static FORCE_INLINE const char* copy_and_filter_macro(parse_state* cs, const char* in, char* out)
+{
+	pp_context* c = cs->context;
+	const char* p = in;
+	char* q = out;
+	const char* start = p;
+	char ch_first = *p;
+
+#if PREPROCESSOR_USE_SSE4_2
+	__m128i k_identifier_needle = _mm_setr_epi8('A', 'Z', 'a', 'z', '0', '9', '_', '_', '$', '$', '$', '$', '$', '$', '$', '$');
+
+	__m128i identifier_word = _mm_loadu_si128((const __m128i*)p);
+	int end_index = _mm_cmpistri(k_identifier_needle, identifier_word, _SIDD_UBYTE_OPS | _SIDD_CMP_RANGES | _SIDD_NEGATIVE_POLARITY);
+
+	_mm_storeu_si128((__m128i*)q, identifier_word);
+	p += end_index;
+	q += end_index;
+
+	// Write as long we had a full word
+	while (end_index == 16)
+	{
+		identifier_word = _mm_loadu_si128((const __m128i*)p);
+		end_index = _mm_cmpistri(k_identifier_needle, identifier_word, _SIDD_UBYTE_OPS | _SIDD_CMP_RANGES | _SIDD_NEGATIVE_POLARITY);
+		if (!end_index)
+		{
+			break;
+		}
+		_mm_storeu_si128((__m128i*)q, identifier_word);
+		p += end_index;
+		q += end_index;
+	}
+
+#else  // PREPROCESSOR_USE_SSE4_2
+
+	// copy the rest of the identifier so we can NUL-terminate it
+	// we copy it to the output buffer since we've already allocated space for it,
+	// and that's where we'll need to copy it if it's not a macro
+	// There is always at least 1 identifier character, so copy the first unconditionally
+	p++;
+	*q++ = ch_first;
+	while (char_is_pp_identifier(*p))
+		*q++ = *p++;
+
+#endif  // !PREPROCESSOR_USE_SSE4_2
+
+	size_t identifier_length = p - start;
+
+	// Test bits for three key characters of the macro by length (first, middle, last)
+	size_t bloom_filter_offset = identifier_length < MACRO_BLOOM_FILTER_SIZE ? identifier_length : MACRO_BLOOM_FILTER_SIZE - 1;
+	if (!(c->macro_bloom_filter[bloom_filter_offset][0] & (1ull << (ch_first & 0x3f))) ||
+		!(c->macro_bloom_filter[bloom_filter_offset][1] & (1ull << (start[identifier_length >> 1] & 0x3f))) ||
+		!(c->macro_bloom_filter[bloom_filter_offset][2] & (1ull << (start[identifier_length - 1] & 0x3f))))
+	{
+		cs->src_offset = p - cs->src;
+		arrsetlennocap(cs->dest, (size_t)(q - cs->dest));  // set the output buffer size to account for the above-copied identifier
+		return p;		// Not a macro, return end of parse
+	}
+
+	// There is an implicit contract that "maybe_expand_macro" should consume identifiers copied by "copy_and_filter_macro".
+	// To allow validation of this, the length starts as zero, is set to non-zero here, then back to zero in "maybe_expand_macro".
+	assert(cs->copied_identifier_length == 0);
+
+	cs->copied_identifier_length = identifier_length;
+	return NULL;		// Could be a macro, return NULL
+}
+
 // Following routine is the workhorse; this copies from a file's text buffer
 // into the output buffer until it finds a directive, a preprocessor Identifier
 // (which might be a macro name), or the end of the file.
@@ -1330,14 +1446,16 @@ static int copy_to_action_point(parse_state* cs)
 	{
 		int prev_state;
 
-		// expand output so there's room to copy entire file unmodified
-		arrsetcap(out, arrlen(out) + (cs->src_length - cs->src_offset) + 1);
-		q = out + arrlen(out);
+		// expand output so there's room to copy entire file unmodified, plus padding so SSE operations can write to the end of the buffer without special cases
+		int out_len = arrlen(out);
+		arrsetcap(out, out_len + (cs->src_length - cs->src_offset) + SSE_READ_PADDING);
+		q = out + out_len;
 
 		{  // make lifetime of line_count unambiguous
 			const int state_limit = cs->state_limit;
 			int line_count = 0;
 
+		resume_parsing:
 			do
 			{
 				uint8 ch = (uint8)*p++;
@@ -1349,9 +1467,31 @@ static int copy_to_action_point(parse_state* cs)
 				prev_state = state;
 				STB_ASSUME(state < PP_STATE_active_count && ch_class < PP_CHAR_CLASS_count);
 				state = pp_transition_table[state][ch_class];
-				line_count += pp_state_is_end_of_line[state];
+				line_count += (ch == '\n') ? 1 : 0;				// pp_state_is_end_of_line[state] -- Assume newlines are normalized, avoid memory read
 			} while (state < state_limit);
 
+			if (state == PP_STATE_saw_identifier_start)
+			{
+				// Copy the identifier and check if it might be a macro.  The first identifier character will already have been written
+				// by the parse loop, so we need to go back one.
+				const char* identifier_start = p - 1;
+				const char* identifier_end = copy_and_filter_macro(cs, identifier_start, q - 1);
+
+				if (identifier_end)
+				{
+					// If not a macro, we can resume parsing
+					q += identifier_end - p;
+					p = identifier_end;
+					state = PP_STATE_ready;
+					cs->in_leading_whitespace = 0;
+					goto resume_parsing;
+				}
+				else
+				{
+					break;
+				}
+			}
+
 			cs->src_line_number += line_count;
 			cs->dest_line_number += line_count;
 		}
@@ -1419,7 +1559,7 @@ static int copy_to_action_point(parse_state* cs)
 	}
 
 	// set output to the amount we actually copied... don't count the copy of the last character
-	arrsetlen(out, (size_t)(q - 1 - out));
+	arrsetlennocap(out, (size_t)(q - 1 - out));
 	cs->dest = out;
 
 	cs->src_offset = (size_t)(p - 1 - cs->src);
@@ -1444,27 +1584,55 @@ static int copy_to_action_point_macro_expansion(parse_state* cs)
 	STB_ASSUME(out != NULL);
 	char* q;
 
-	// expand output so copying entire file unmodified
-	arrsetcap(out, arrlen(out) + (cs->src_length - cs->src_offset) + 1);
-	q = out + arrlen(out);
+	// expand output so copying entire file unmodified is possible
+	int out_len = arrlen(out);
+	arrsetcap(out, out_len + (cs->src_length - cs->src_offset) + SSE_READ_PADDING);
+	q = out + out_len;
 
 	STB_ASSUME(q != NULL);
 
 	state = PP_STATE_ready;
 
-	do
+	for (;;)
 	{
-		uint8 ch = (uint8)*p++;
-		uint8 ch_class;
-		*q++ = ch;
-		ch_class = pp_char_class[ch];
-		state = pp_transition_table[state][ch_class];
-	} while (state < PP_STATE_active_count);
+		do
+		{
+			uint8 ch = (uint8)*p++;
+			uint8 ch_class;
+			*q++ = ch;
+			ch_class = pp_char_class[ch];
+			state = pp_transition_table[state][ch_class];
+		} while (state < PP_STATE_active_count);
+
+		if (state == PP_STATE_saw_identifier_start)
+		{
+			// Copy the identifier and check if it might be a macro.  The first identifier character will already have been written
+			// by the parse loop, so we need to go back one.
+			const char* identifier_start = p - 1;
+			const char* identifier_end = copy_and_filter_macro(cs, identifier_start, q - 1);
+
+			if (identifier_end)
+			{
+				q += identifier_end - p;
+				p = identifier_end;
+				state = PP_STATE_ready;
+				cs->in_leading_whitespace = 0;
+			}
+			else
+			{
+				break;
+			}
+		}
+		else
+		{
+			break;
+		}
+	}
 
 	cs->src_offset = (size_t)(p - 1 - cs->src);
 
 	// set output to the amount we actually copied... don't count the copy of the last character
-	arrsetlen(out, (size_t)(q - 1 - out));
+	arrsetlennocap(out, (size_t)(q - 1 - out));
 	cs->dest = out;
 
 	return state;
@@ -1476,6 +1644,109 @@ static int copy_to_action_point_macro_expansion(parse_state* cs)
 static const char* scan_to_directive(const char* p, int* p_line_number)
 {
 	int line_number = *p_line_number;
+
+#if PREPROCESSOR_USE_SSE4_2
+	// Fast loop assumes comments have been stripped, and newlines normalized.  Searches for a # not in quoted text.  Range is the
+	// negation of the characters we want, which will also stop at the null terminator.  Around 10x faster than original loop.
+	// Note that the source may still contain comments added via DumpShaderDefinesAsCommentedCode, but those defines are pasted
+	// in externally, and will be outside of any directive, and therefore not break this code.
+	__m128i k_directive_needle_negated = _mm_setr_epi8(
+		0    + 1, '\"' - 1,		// null         to double quote
+		'#'  + 1, '\'' - 1,		// hash         to single quote
+		'\'' + 1, 255,			// single quote to rest of characters
+		0,0,0,0,0,0,0,0,0,0);
+
+	__m128i k_newlines = _mm_set1_epi8('\n');
+
+	for (;;)
+	{
+		__m128i scan_word = _mm_loadu_si128((const __m128i*)p);
+		int newline_mask = _mm_movemask_epi8(_mm_cmpeq_epi8(k_newlines, scan_word));
+
+		// Check for full word with no relevant characters (common case)
+		if (!_mm_cmpistrc(k_directive_needle_negated, scan_word, _SIDD_UBYTE_OPS | _SIDD_CMP_RANGES | _SIDD_NEGATIVE_POLARITY))
+		{
+			line_number += _mm_popcnt_u32(newline_mask);
+			p += 16;
+			continue;
+		}
+
+		// Find the matching character
+		int end_index = _mm_cmpistri(k_directive_needle_negated, scan_word, _SIDD_UBYTE_OPS | _SIDD_CMP_RANGES | _SIDD_NEGATIVE_POLARITY);
+		newline_mask = newline_mask & (0xffff >> (16 - end_index));
+		p += end_index;
+		line_number += _mm_popcnt_u32(newline_mask);
+
+		char found_char = *p;
+		if (found_char == '#')
+		{
+			// Scan backwards for a newline.  Writing this non-SSE, with the theory that most directives are at the start of a line,
+			// and so testing a single character is all that's required 99% of the time, and an SSE string instruction is much slower
+			// for that common case.  We are in a commented out #if block, so there must be a newline related to the #if block earlier
+			// in the file, so we don't need to check for beginning of string.
+			for (const char* leading_whitespace_scan = p - 1;; leading_whitespace_scan--)
+			{
+				if (*leading_whitespace_scan == '\n')
+				{
+					// Leading whitespace hash!  Return result.
+					*p_line_number = line_number;
+					return p;
+				}
+				else if (pp_char_class[*leading_whitespace_scan] != PP_CHAR_CLASS_whitespace)
+				{
+					// Not a leading hash!  Continue parsing
+					p++;
+					break;
+				}
+			}
+		}
+		else if (found_char == '\"' || found_char == '\'')
+		{
+			// Quoted text.  Not bothering to SSE optimize this, as it is rarely reached.  Directives containing text aren't scanned
+			// in this loop, and other sources of text are rare.
+			char current_char;
+			for (p++; (current_char = *p); p++)
+			{
+				if (current_char == found_char)
+				{
+					// Found close quote, skip it and break out
+					p++;
+					break;
+				}
+				else if (current_char == '\\')
+				{
+					// Escape, skip the next character unconditionally if not EOF
+					if (p[1] != 0)
+					{
+						p++;
+						current_char = *p;
+					}
+				}
+
+				// Track newlines (including escaped newlines)
+				if (current_char == '\n')
+				{
+					line_number++;
+				}
+			}
+
+			if (!current_char)
+			{
+				// Unterminated string
+				*p_line_number = line_number;
+				return p;
+			}
+		}
+		else
+		{
+			// Null terminator
+			*p_line_number = line_number;
+			return p;
+		}
+	}
+
+#else  // PREPROCESSOR_USE_SSE4_2
+
 	int state = PP_STATE_ready, prev_state;
 	for (;;)
 	{
@@ -1506,6 +1777,8 @@ static const char* scan_to_directive(const char* p, int* p_line_number)
 		else
 			state = PP_STATE_ready;
 	}
+
+#endif  // !PREPROCESSOR_USE_SSE4_2
 }
 
 //
@@ -1679,6 +1952,9 @@ static char* copy_line_without_comments(char* buffer, size_t buffer_size, const
 				if (p[0] == 0 || p[-1] != '\\')
 				{
 					arrput(out, 0);
+#if SSE_READ_PADDING
+					arrsetcap(out, arrlennonull(out) + SSE_READ_PADDING);
+#endif
 					*endp = p;
 					return out;
 				}
@@ -1694,6 +1970,9 @@ static char* copy_line_without_comments(char* buffer, size_t buffer_size, const
 					while (!char_is_end_of_line(*p))
 						++p;
 					arrput(out, 0);
+#if SSE_READ_PADDING
+					arrsetcap(out, arrlennonull(out) + SSE_READ_PADDING);
+#endif
 					*endp = p;
 					return out;
 				}
@@ -1883,8 +2162,8 @@ static int parse_directive_after_hash(const char** ptr_p, int* newlines)
 	while (char_is_pp_identifier_first(*p))
 		sum += (uint8)*p++;
 	sum &= (DIRECTIVE_HASH_SIZE - 1);
-	// @OPTIMIZE: use a masked 8-byte comparison?
-	if (0 == memcmp(start, directive_hash[sum].name, directive_hash[sum].name_len) && !char_is_pp_identifier(directive_hash[sum].name_len))
+
+	if (fast_preprocess_string_equals_short(start, directive_hash[sum].name, directive_hash[sum].name_len) && !char_is_pp_identifier(directive_hash[sum].name_len))
 	{
 		*ptr_p = p;
 		return sum;
@@ -2080,9 +2359,9 @@ static void preprocess_string(parse_state* cs, int in_macro_expansion, char* in_
 					int i, stop = 0;
 					for (i = (int)arrlen(c->ifdef_stack) - 1; i >= 0; --i)
 					{
-						if (do_error(cs, "No #endif found for previous #%s", directive_hash[c->ifdef_stack[i].type].name))
+						if (do_error(cs, "No #endif found for previous #%s", &directive_hash[c->ifdef_stack[i].type].name[0]))
 							stop = 1;
-						error_supplement(cs, PP_RESULT_ERROR, c->ifdef_stack[i].line, "Location of #%s", directive_hash[c->ifdef_stack[i].type].name);
+						error_supplement(cs, PP_RESULT_ERROR, c->ifdef_stack[i].line, "Location of #%s", &directive_hash[c->ifdef_stack[i].type].name[0]);
 					}
 					if (stop)
 						return;
@@ -2119,18 +2398,12 @@ static void update_macro_filter(pp_context* c, char* identifier, size_t identifi
 {
 	assert(identifier[0] != 0 && identifier_length > 0);
 
-#if MACRO_BLOOM_FILTER_ENABLED
 	size_t bloom_filter_offset = identifier_length < MACRO_BLOOM_FILTER_SIZE ? identifier_length : MACRO_BLOOM_FILTER_SIZE - 1;
 
 	// Set bits for three key characters of the macro by length (first, middle, last)
 	c->macro_bloom_filter[bloom_filter_offset][0] |= (1ull << (identifier[0] & 0x3f));
 	c->macro_bloom_filter[bloom_filter_offset][1] |= (1ull << (identifier[identifier_length >> 1] & 0x3f));
 	c->macro_bloom_filter[bloom_filter_offset][2] |= (1ull << (identifier[identifier_length - 1] & 0x3f));
-
-#else
-	if (identifier[1] == 0 || identifier[2] == 0)
-		c->macro_name_filter[(uint8)identifier[0]][(uint8)identifier[1]] = 1;
-#endif
 }
 
 static struct macro_definition* create_macro_definition(parse_state* ps, const char* def)
@@ -2710,57 +2983,15 @@ static void maybe_expand_macro(parse_state* cs, struct macro_definition* pending
 	}
 	else
 	{
-		const char* start = p;
+		// Identifier will have been copied and length set by copy_and_filter_macro
+		identifier_length = cs->copied_identifier_length;
+		cs->copied_identifier_length = 0;
 
-#if !MACRO_BLOOM_FILTER_ENABLED
-		uint8 c0 = (uint8)p[0];
-		uint8 c1 = (uint8)p[1];
-
-		if (!char_is_pp_identifier(c1))
-		{
-			// 1-character identifier
-			if (c->macro_name_filter[c0][0] == 0)
-			{
-				*q++ = c0;
-				p += 1;
-				goto not_macro;
-			}
-		}
-		else if (!char_is_pp_identifier(p[2]))
-		{
-			// 2-character identifier
-			// could use bloom filter instead, but while this is 64KB, only portions of it will ever
-			if (c->macro_name_filter[c0][c1] == 0)
-			{
-				*q++ = c0;
-				*q++ = c1;
-				p += 2;
-				goto not_macro;
-			}
-		}
-#endif
+		// "maybe_expand_macro" consumes identifiers originally copied by "copy_and_filter_macro" -- make sure one was copied
+		assert(identifier_length > 0);
 
-		// copy the rest of the identifier so we can NUL-terminate it
-		// we copy it to the output buffer since we've already allocated space for it,
-		// and that's where we'll need to copy it if it's not a macro
-		// There is always at least 1 identifier character, so copy the first unconditionally
-		char ch_first = *p++;
-		*q++ = ch_first;
-		while (char_is_pp_identifier(*p))
-			*q++ = *p++;
-
-		identifier_length = p - start;
-
-#if MACRO_BLOOM_FILTER_ENABLED
-		// Test bits for three key characters of the macro by length (first, middle, last)
-		size_t bloom_filter_offset = identifier_length < MACRO_BLOOM_FILTER_SIZE ? identifier_length : MACRO_BLOOM_FILTER_SIZE - 1;
-		if (!(c->macro_bloom_filter[bloom_filter_offset][0] & (1ull << ch_first)) ||
-			!(c->macro_bloom_filter[bloom_filter_offset][1] & (1ull << start[identifier_length >> 1])) ||
-			!(c->macro_bloom_filter[bloom_filter_offset][2] & (1ull << start[identifier_length - 1])))
-		{
-			goto not_macro;
-		}
-#endif
+		p += identifier_length;
+		q += identifier_length;
 
 		md = (struct macro_definition*)stringhash_get(&c->macro_map, identifier, identifier_length);
 		if (md == NULL || md->disabled)
@@ -3019,12 +3250,27 @@ static void maybe_expand_macro(parse_state* cs, struct macro_definition* pending
 				}
 				*e = 0;
 
+#if SSE_READ_PADDING
+				// Need to ensure padding for safe SSE reads without special cases -- normally this comes from a stack allocated buffer with plenty
+				// of padding, so it won't need to reallocate, but we need to handle edge cases.  We need to check padding before the
+				// stbds_arrinline_suballoc_char call, which trims the allocation.
+				int copy_len = arrlennonull(copy);
+				int copy_capacity_padding = arrcapnonull(copy) - copy_len;
+#endif
+
 				// Sub-allocate the remaining space in argument_buffer for subsequent arguments.  Note that it's necessary to do the suballocation BEFORE
 				// the call to arrsetlen below, because downstream code assumes there is a null terminator beyond the official length of the array.  Doing
 				// the sub-allocation first preserves the null terminator -- without this, the sub-allocation will overwrite it.  We only suballocate if
 				// a minimum of 200 characters is available, as copy_argument reserves this.
 				stbds_arrinline_suballoc_char(argument_buffer, 200);
 
+#if SSE_READ_PADDING
+				if (copy_capacity_padding < SSE_READ_PADDING)
+				{
+					arrsetcap(copy, copy_len + SSE_READ_PADDING);
+				}
+#endif
+
 				arrsetlen(copy, e - p);
 				arrput(arguments, copy);
 				cs->src_line_number += arg_newlines;
@@ -3100,6 +3346,16 @@ static void maybe_expand_macro(parse_state* cs, struct macro_definition* pending
 					e = copy;
 				*e = 0;
 
+#if SSE_READ_PADDING
+				// Need to ensure padding for safe SSE reads without special cases -- normally this comes from a stack allocated buffer with
+				// plenty of padding, so this will have sufficient padding, but we need to handle edge cases.
+				int copy_len = arrlennonull(copy);
+				if (arrcapnonull(copy) - copy_len < SSE_READ_PADDING)
+				{
+					arrsetcap(copy, copy_len + SSE_READ_PADDING);
+				}
+#endif
+
 				arrput(arguments, copy);
 				cs->src_line_number += arg_newlines;
 			}
@@ -3213,6 +3469,9 @@ static void maybe_expand_macro(parse_state* cs, struct macro_definition* pending
 		}
 
 		arrput(tmp, 0);
+#if SSE_READ_PADDING
+		arrsetcap(tmp, arrlen(tmp) + SSE_READ_PADDING);
+#endif
 
 		{
 			for (i = 0; i < arrlen(arguments); ++i)
@@ -3444,11 +3703,11 @@ static void process_directive(parse_state* cs, conditional_state* cons)
 	pp_context* c = cs->context;
 	const char *p, *endptr;
 	const char* alloc = NULL;	 // all exit paths SHOULD do 'if (alloc) arrfree(alloc)', but we allow it to leak on error-handling returns
-	char buffer[2048];
+	char buffer[2048 + SSE_READ_PADDING];
 	int dummy;	// uninitialized, should never get incremented by parse_directive_after_hash
 
 	assert(cs->src[cs->src_offset] == '#');
-	p = copy_line_without_comments(buffer, sizeof(buffer), cs->src + cs->src_offset + 1, &cs->src_line_number, &endptr);
+	p = copy_line_without_comments(buffer, sizeof(buffer) - SSE_READ_PADDING, cs->src + cs->src_offset + 1, &cs->src_line_number, &endptr);
 	if (p != buffer)
 		alloc = p;
 
@@ -4077,8 +4336,8 @@ static void process_directive(parse_state* cs, conditional_state* cons)
 					int i;
 					for (i = (int)arrlen(c->ifdef_stack) - 1; i >= cs->conditional_nesting_depth_at_start; --i)
 					{
-						do_error(cs, "No #endif found for previous #%s", directive_hash[c->ifdef_stack[i].type].name);
-						error_supplement(cs, PP_RESULT_ERROR, c->ifdef_stack[i].line, "Location of #%s", directive_hash[c->ifdef_stack[i].type].name);
+						do_error(cs, "No #endif found for previous #%s", &directive_hash[c->ifdef_stack[i].type].name[0]);
+						error_supplement(cs, PP_RESULT_ERROR, c->ifdef_stack[i].line, "Location of #%s", &directive_hash[c->ifdef_stack[i].type].name[0]);
 					}
 				}
 				return;
@@ -4111,7 +4370,7 @@ static void process_directive(parse_state* cs, conditional_state* cons)
 					{
 						if (cons->conditionally_disable == CONDITIONAL_skip_current_block)
 						{
-							p = copy_line_without_comments(buffer, sizeof(buffer), p, &cs->src_line_number, &endptr);
+							p = copy_line_without_comments(buffer, sizeof(buffer) - SSE_READ_PADDING, p, &cs->src_line_number, &endptr);
 							if (p != buffer)
 								alloc = p;
 							if (endptr == NULL)
@@ -4136,7 +4395,7 @@ static void process_directive(parse_state* cs, conditional_state* cons)
 				case HASH_endif:
 					if (disable_nesting_level == 0)
 					{
-						p = copy_line_without_comments(buffer, sizeof(buffer), p, &cs->src_line_number, &endptr);
+						p = copy_line_without_comments(buffer, sizeof(buffer) - SSE_READ_PADDING, p, &cs->src_line_number, &endptr);
 						if (p != buffer)
 							alloc = p;
 						if (endptr == NULL)
@@ -4156,7 +4415,7 @@ static void process_directive(parse_state* cs, conditional_state* cons)
 					{
 						if (cons->conditionally_disable == CONDITIONAL_skip_current_block)
 						{
-							p = copy_line_without_comments(buffer, sizeof(buffer), p, &cs->src_line_number, &endptr);
+							p = copy_line_without_comments(buffer, sizeof(buffer) - SSE_READ_PADDING, p, &cs->src_line_number, &endptr);
 							if (p != buffer)
 								alloc = p;
 							if (endptr == NULL)
@@ -4200,6 +4459,7 @@ static char* preprocess_string_from_file(parse_state* ps, const char* filename,
 	cs.filename = filename;
 	cs.last_output_filename = "";  // force #line directive to include filename
 	cs.src = text;
+	cs.copied_identifier_length = 0;
 	cs.src_line_number = 1;
 	cs.dest_line_number = 1;
 	cs.in_leading_whitespace = 1;
@@ -4546,15 +4806,7 @@ char* preprocess_file(const char* filename,
 	// backfire by slowing down the initialization
 	stringhash_create(&c.macro_map, 4096);
 
-#if MACRO_BLOOM_FILTER_ENABLED
 	memset(c.macro_bloom_filter, 0, sizeof(c.macro_bloom_filter));
-#else
-	c.macro_name_filter = STB_COMMON_MALLOC(65536);
-	if (c.macro_name_filter == NULL)
-		return NULL;
-
-	memset(c.macro_name_filter, 0, 65536);
-#endif
 
 	for (i = 0; i < arrlen(predefined_macros); ++i)
 		define_macro(&c.macro_map, predefined_macros[i]);
@@ -4596,9 +4848,6 @@ char* preprocess_file(const char* filename,
 	shfree(c.undef_map);
 	arrfree(c.ifdef_stack);
 	stringhash_destroy(&c.macro_map);
-#if !MACRO_BLOOM_FILTER_ENABLED
-	STB_COMMON_FREE(c.macro_name_filter);
-#endif
 	stb_arena_free(&c.macro_arena);
 
 	if (c.diagnostics)
@@ -4660,13 +4909,13 @@ static void init_directive(char* s, int hash)
 		fprintf(stderr, "Preprocessor fatal internal error: hash initialization mismatch for '%s' (%s %d).\n", s, __FILE__, __LINE__);
 		exit(1);
 	}
-	if (directive_hash[hash].name != 0)
+	if (directive_hash[hash].name[0] != 0)
 	{
 		fprintf(stderr, "Preprocessor fatal internal error: hash collision between '%s' and '%s' (%s %d).\n", s, directive_hash[hash].name, __FILE__, __LINE__);
 		exit(1);
 	}
 
-	directive_hash[hash].name = s;
+	strcpy_s(directive_hash[hash].name, sizeof(directive_hash[hash].name), s);
 	directive_hash[hash].name_len = strlen(s);
 }
 
diff --git a/Engine/Source/Developer/ShaderPreprocessor/Private/stb_preprocess/stb_alloc.c b/Engine/Source/Developer/ShaderPreprocessor/Private/stb_preprocess/stb_alloc.c
--- a/Engine/Source/Developer/ShaderPreprocessor/Private/stb_preprocess/stb_alloc.c
+++ b/Engine/Source/Developer/ShaderPreprocessor/Private/stb_preprocess/stb_alloc.c
@@ -41,8 +41,8 @@ void* stb_arena_alloc_aligned(struct stb_arena* a, size_t size, size_t align)
 			}
 		}
 
-		// allocate and record allocation
-		a->last_block = (unsigned char*)STB_COMMON_MALLOC(a->last_block_size + 31);
+		// allocate and record allocation.  31 bytes for alignment, 16 bytes for SSE padding at end of a block
+		a->last_block = (unsigned char*)STB_COMMON_MALLOC(a->last_block_size + 31 + 16);
 		STB_ASSUME(a->blocks != NULL);
 		a->blocks[a->num_blocks++] = a->last_block;
 		a->last_block_alloc_offset = 0;
diff --git a/Engine/Source/Developer/ShaderPreprocessor/Private/stb_preprocess/stb_ds.h b/Engine/Source/Developer/ShaderPreprocessor/Private/stb_preprocess/stb_ds.h
--- a/Engine/Source/Developer/ShaderPreprocessor/Private/stb_preprocess/stb_ds.h
+++ b/Engine/Source/Developer/ShaderPreprocessor/Private/stb_preprocess/stb_ds.h
@@ -392,6 +392,7 @@ CREDITS
 #define arraddnptr stbds_arraddnptr
 #define arraddnindex stbds_arraddnindex
 #define arrsetlen stbds_arrsetlen
+#define arrsetlennocap stbds_arrsetlennocap
 #define arrinitlen stbds_arrinitlen
 #define arrlast stbds_arrlast
 #define arrins stbds_arrins
@@ -400,6 +401,7 @@ CREDITS
 #define arrdeln stbds_arrdeln
 #define arrdelswap stbds_arrdelswap
 #define arrcap stbds_arrcap
+#define arrcapnonull stbds_arrcapnonull
 #define arrsetcap stbds_arrsetcap
 
 #define hmput stbds_hmput
@@ -558,8 +560,10 @@ extern "C"
 
 #define stbds_arrsetcap(a, n) ((!(a) || (n) > stbds_header(a)->capacity) ? (stbds_arrgrow(a, 0, n), 0) : 0)
 #define stbds_arrsetlen(a, n) ((stbds_arrcap(a) < (size_t)(n) ? stbds_arrsetcap((a), (size_t)(n)), 0 : 0), (a) ? stbds_header(a)->length = (size_t)(n) : 0)
+#define stbds_arrsetlennocap(a, n) (stbds_header(a)->length = n)
 #define stbds_arrinitlen(a, n) ((stbds_arrsetcap((a), (size_t)(n)), stbds_header(a)->length = (size_t)(n)))
 #define stbds_arrcap(a) ((a) ? stbds_header(a)->capacity : 0)
+#define stbds_arrcapnonull(a) (stbds_header(a)->capacity)
 #define stbds_arrlen(a) ((a) ? (ptrdiff_t)stbds_header(a)->length : 0)
 #define stbds_arrlennonull(a) ((ptrdiff_t)stbds_header(a)->length)
 #define stbds_arrlenu(a) ((a) ? stbds_header(a)->length : 0)
diff --git a/Engine/Source/Runtime/RenderCore/Private/ShaderCore.cpp b/Engine/Source/Runtime/RenderCore/Private/ShaderCore.cpp
--- a/Engine/Source/Runtime/RenderCore/Private/ShaderCore.cpp
+++ b/Engine/Source/Runtime/RenderCore/Private/ShaderCore.cpp
@@ -1841,16 +1841,174 @@ void ShaderConvertAndStripComments(const FString& ShaderSource, TArray<ANSICHAR>
 	// STB preprocessor does not strip comments, so we do so here before returning the loaded source
 	// Doing so is barely more costly than the memcopy we require anyways so has negligible overhead.
 	// Reserve worst case (i.e. assuming there are no comments at all) to avoid reallocation
-	int32 BufferSize = ShaderSource.Len() + 1; // +1 to append null terminator
+	int32 BufferSize = ShaderSource.Len() + 16;		// need extra for null terminator plus padding for SSE read operations at the end of the buffer
 	OutStripped.SetNumUninitialized(BufferSize);
 
 	ANSICHAR* CurrentOut = OutStripped.GetData();
 
-	const TCHAR* const End = ShaderSource.GetCharArray().GetData() + ShaderSource.Len();
+	const TCHAR* const Start = ShaderSource.GetCharArray().GetData();
+	const TCHAR* const End = Start + ShaderSource.Len();
 
 	// We rely on null termination to avoid the need to check Current < End in some cases
 	check(*End == TEXT('\0'));
-	for (const TCHAR* Current = ShaderSource.GetCharArray().GetData(); Current < End;)
+
+	const TCHAR* Current = Start;
+
+#if PLATFORM_ALWAYS_HAS_SSE4_2
+	__m128i CharCR = _mm_set1_epi8('\r');			// Carriage return
+	__m128i CharLF = _mm_set1_epi8('\n');			// Line feed (newline)
+	__m128i CharSlash = _mm_set1_epi8('/');
+	__m128i CharStar = _mm_set1_epi8('*');
+
+	// We process 15 characters at a time, so we can find comment starts (needs access to pairs of characters)
+	const TCHAR* EndSse = End - 16;
+	for (; Current < EndSse; )
+	{
+		__m128i First8 = _mm_loadu_epi16(Current);
+		__m128i Second8 = _mm_loadu_epi16(Current + 8);
+		__m128i CurrentWord = _mm_packus_epi16(First8, Second8);
+
+		int32 CRMask = _mm_movemask_epi8(_mm_cmpeq_epi8(CurrentWord, CharCR));
+		int32 SlashMask = _mm_movemask_epi8(_mm_cmpeq_epi8(CurrentWord, CharSlash));
+		int32 StarMask = _mm_movemask_epi8(_mm_cmpeq_epi8(CurrentWord, CharStar));
+
+		// If we encounter a carriage return, fall back to slower single character path that handles CR/LF combos
+		if (CRMask)
+		{
+			// Go back one character if first character in current word is CR, and previous character was LF, so
+			// the single character parser can treat it as a newline pair.
+			if ((CRMask & 1) && (Current > Start) && *(Current - 1) == '\n')
+			{
+				Current--;
+				CurrentOut--;
+			}
+			break;
+		}
+
+		// Echo the current word
+		_mm_storeu_epi8(CurrentOut, CurrentWord);
+
+		// Check if there is a comment start, meaning a slash followed by slash or star, which we can detect by shifting right
+		// a mask containing both slash and star, and seeing if that overlaps with a slash.
+		int32 CommentStartMask = SlashMask & ((SlashMask | StarMask) >> 1);
+		if (!CommentStartMask)
+		{
+			// If no potential comment start, advance 15 characters and parse again
+			CurrentOut += 15;
+			Current += 15;
+			continue;
+		}
+
+		// Advance input to contents of comment, output to end of non-comment characters
+		int32 CommentOffset = _tzcnt_u32(CommentStartMask);
+		Current += CommentOffset + 2;
+		CurrentOut += CommentOffset;
+
+		if (*(Current - 1) == '/')
+		{
+			// Single line comment, advance to newline
+			bool bFoundNewline = false;
+
+			for (; Current < EndSse;)
+			{
+				First8 = _mm_loadu_epi16(Current);
+				Second8 = _mm_loadu_epi16(Current + 8);
+				CurrentWord = _mm_packus_epi16(First8, Second8);
+
+				CRMask = _mm_movemask_epi8(_mm_cmpeq_epi8(CurrentWord, CharCR));
+				int32 LFMask = _mm_movemask_epi8(_mm_cmpeq_epi8(CurrentWord, CharLF));
+				int32 EitherMask = CRMask | LFMask;
+				if (EitherMask)
+				{
+					int32 NewlineOffset = _tzcnt_u32(EitherMask);
+					Current += NewlineOffset;
+					bFoundNewline = true;
+					break;
+				}
+				else
+				{
+					Current += 16;
+				}
+			}
+
+			if (!bFoundNewline)
+			{
+				// Ran out of input buffer we can safely scan with SSE -- resume comment parsing in single character parser.
+				goto SingleLineCommentParse;
+			}
+			if (CRMask)
+			{
+				// Hit a CR.  Stop and fall back to single character parser.  Note that we don't need to worry about rewinding for
+				// a newline pair here, because we stop on either that's encountered first, so we haven't emitted a newline yet.
+				break;
+			}
+		}
+		else
+		{
+			// Multi line comment, skip to end of comment, writing newlines
+			bool bFoundEnd = false;
+
+			for (; Current < EndSse;)
+			{
+				First8 = _mm_loadu_epi16(Current);
+				Second8 = _mm_loadu_epi16(Current + 8);
+				CurrentWord = _mm_packus_epi16(First8, Second8);
+
+				// Fall back to single character parsing if we hit a CR
+				CRMask = _mm_movemask_epi8(_mm_cmpeq_epi8(CurrentWord, CharCR));
+				if (CRMask)
+				{
+					// Go back one character if this is the first CR, and previous character was LF
+					if ((CRMask & 1) && (Current > Start) && *(Current - 1) == '\n')
+					{
+						Current--;
+						CurrentOut--;
+					}
+					goto MultiLineCommentParse;
+				}
+
+				StarMask = _mm_movemask_epi8(_mm_cmpeq_epi8(CurrentWord, CharStar));
+				SlashMask = _mm_movemask_epi8(_mm_cmpeq_epi8(CurrentWord, CharSlash));
+				int32 LFMask = _mm_movemask_epi8(_mm_cmpeq_epi8(CurrentWord, CharLF));
+
+				int32 CommentEndMask = StarMask & (SlashMask >> 1);
+				if (CommentEndMask)
+				{
+					// Process any newlines before the comment end
+					int32 CommentEndOffset = _tzcnt_u32(CommentEndMask);
+					LFMask &= (0xffff >> (16 - CommentEndOffset));
+					if (LFMask)
+					{
+						_mm_storeu_epi8(CurrentOut, CharLF);
+						CurrentOut += _mm_popcnt_u32(LFMask);
+					}
+					Current += CommentEndOffset + 2;
+					bFoundEnd = true;
+					break;
+				}
+				else
+				{
+					// No comment end -- process any newlines in the first 15 characters and continue
+					LFMask &= 0x7fff;
+					if (LFMask)
+					{
+						_mm_storeu_epi8(CurrentOut, CharLF);
+						CurrentOut += _mm_popcnt_u32(LFMask);
+					}
+					Current += 15;
+				}
+			}
+
+			if (!bFoundEnd)
+			{
+				// Ran out of input buffer we can safely scan with SSE -- resume comment parsing in single character parser.
+				goto MultiLineCommentParse;
+			}
+		}
+	}
+#endif	// PLATFORM_ALWAYS_HAS_SSE4_2
+
+	for (; Current < End;)
 	{
 		// sanity check that we're not overrunning the buffer
 		check(CurrentOut < (OutStripped.GetData() + BufferSize));
@@ -1872,6 +2030,7 @@ void ShaderConvertAndStripComments(const FString& ShaderSource, TArray<ANSICHAR>
 		{
 			if (Current[1] == '/')
 			{
+				SingleLineCommentParse:
 				while (!IsEndOfLine(*Current) && Current < End)
 				{
 					++Current;
@@ -1882,6 +2041,7 @@ void ShaderConvertAndStripComments(const FString& ShaderSource, TArray<ANSICHAR>
 				Current += 2;
 				while (Current < End)
 				{
+					MultiLineCommentParse:
 					if (Current[0] == '*' && Current[1] == '/')
 					{
 						Current += 2;
@@ -1908,7 +2068,7 @@ void ShaderConvertAndStripComments(const FString& ShaderSource, TArray<ANSICHAR>
 	check(CurrentOut < (OutStripped.GetData() + BufferSize));
 	*CurrentOut++ = 0;
 
-	// Set correct length after stripping but don't bother shrinking/reallocating, minor memory overhead to save time
+	// Set correct length after stripping but don't bother shrinking/reallocating -- shrinking would remove SSE padding we require
 	OutStripped.SetNum(CurrentOut - OutStripped.GetData(), /* bAllowShrinking */false);
 }
 
