Shader Preprocessor:   Add padding for SSE reads as actual zeroes, rather than capacity padding.  Capacity padding doesn't work for compiler workers some of the time.

diff --git a/Engine/Source/Developer/ShaderPreprocessor/Private/ShaderPreprocessor.cpp b/Engine/Source/Developer/ShaderPreprocessor/Private/ShaderPreprocessor.cpp
--- a/Engine/Source/Developer/ShaderPreprocessor/Private/ShaderPreprocessor.cpp
+++ b/Engine/Source/Developer/ShaderPreprocessor/Private/ShaderPreprocessor.cpp
@@ -103,7 +103,6 @@ struct FStbLoadedInclude
 {
 	const ANSICHAR* Data = nullptr;				// Points to SharedData, LocalData, or data from FShaderCompilerEnvironment
 	size_t DataLength = 0;
-	size_t DataCapacity = 0;
 	FShaderSharedAnsiStringPtr SharedData;
 	TArray<ANSICHAR> LocalData;
 };
@@ -140,7 +139,6 @@ static const ANSICHAR* StbLoadFile(const ANSICHAR* Filename, void* RawContext, s
 
 			ContentsCached->Data = ContentsCached->LocalData.GetData();
 			ContentsCached->DataLength = ContentsCached->LocalData.Num();
-			ContentsCached->DataCapacity = ContentsCached->LocalData.Max();
 		}
 		else
 		{
@@ -150,7 +148,6 @@ static const ANSICHAR* StbLoadFile(const ANSICHAR* Filename, void* RawContext, s
 			{
 				ContentsCached->Data = InMemorySourceAnsi->Get()->GetData();
 				ContentsCached->DataLength = InMemorySourceAnsi->Get()->Num();
-				ContentsCached->DataCapacity = InMemorySourceAnsi->Get()->Max();
 			}
 			else
 			{
@@ -159,14 +156,17 @@ static const ANSICHAR* StbLoadFile(const ANSICHAR* Filename, void* RawContext, s
 
 				ContentsCached->Data = ContentsCached->SharedData->GetData();
 				ContentsCached->DataLength = ContentsCached->SharedData->Num();
-				ContentsCached->DataCapacity = ContentsCached->SharedData->Max();
 			}
 		}
 
 		// Need 15 characters beyond null terminator, so an unaligned SSE read at the null terminator can safely read 15 extra unused characters
-		// without going out of memory bounds.  ShaderConvertAndStripComments ensures this padding.  We could optionally allocate (or reallocate)
-		// a local copy as a fallback to handle this case without asserting, but it would be a silent performance degradation.
-		checkf(ContentsCached->DataCapacity >= ContentsCached->DataLength + 15, TEXT("Shader preprocessor ANSI files must include 15 bytes of capacity padding past null terminator"));
+		// without going out of memory bounds.  ShaderConvertAndStripComments adds this padding in the form of extra trailing zeroes.  Make sure
+		// these zeroes are there.
+		static const char SixteenZeroes[16] = { 0 };
+		checkf(ContentsCached->DataLength >= 16 && memcmp(&ContentsCached->Data[ContentsCached->DataLength - 16], SixteenZeroes, 16) == 0,
+			TEXT("Shader preprocessor ANSI files must include 15 bytes of zero padding past null terminator"));
+
+		ContentsCached->DataLength -= 15;
 	}
 	check(ContentsCached);
 	*OutLength = ContentsCached->DataLength;
diff --git a/Engine/Source/Runtime/RenderCore/Private/ShaderCore.cpp b/Engine/Source/Runtime/RenderCore/Private/ShaderCore.cpp
--- a/Engine/Source/Runtime/RenderCore/Private/ShaderCore.cpp
+++ b/Engine/Source/Runtime/RenderCore/Private/ShaderCore.cpp
@@ -2064,11 +2064,14 @@ void ShaderConvertAndStripComments(const FString& ShaderSource, TArray<ANSICHAR>
 			}
 		}
 	}
-	// Null terminate after comment-stripped copy
-	check(CurrentOut < (OutStripped.GetData() + BufferSize));
-	*CurrentOut++ = 0;
+	// Null terminate after comment-stripped copy, plus 15 zero padding characters for SSE safe reads
+	check(CurrentOut + 16 <= (OutStripped.GetData() + BufferSize));
+	for (int32 TerminateAndPadIndex = 0; TerminateAndPadIndex < 16; TerminateAndPadIndex++)
+	{
+		*CurrentOut++ = 0;
+	}
 
-	// Set correct length after stripping but don't bother shrinking/reallocating -- shrinking would remove SSE padding we require
+	// Set correct length after stripping but don't bother shrinking/reallocating, minor memory overhead to save time
 	OutStripped.SetNum(CurrentOut - OutStripped.GetData(), /* bAllowShrinking */false);
 }
 
