Shader Preprocessor platform specific compile error / warning fixes.  Unused branch label on non-SSE platforms, char array index, conversion of 255 to char in SSE vector initialization (used by unsigned compare intrinsic so code was correct), strcpy_s not supported on some platforms -- switched to strncpy.

#rnx

diff --git a/Engine/Source/Developer/ShaderPreprocessor/Private/stb_preprocess/preprocessor.c b/Engine/Source/Developer/ShaderPreprocessor/Private/stb_preprocess/preprocessor.c
--- a/Engine/Source/Developer/ShaderPreprocessor/Private/stb_preprocess/preprocessor.c
+++ b/Engine/Source/Developer/ShaderPreprocessor/Private/stb_preprocess/preprocessor.c
@@ -1653,7 +1653,7 @@ static const char* scan_to_directive(const char* p, int* p_line_number)
 	__m128i k_directive_needle_negated = _mm_setr_epi8(
 		0    + 1, '\"' - 1,		// null         to double quote
 		'#'  + 1, '\'' - 1,		// hash         to single quote
-		'\'' + 1, 255,			// single quote to rest of characters
+		'\'' + 1, (char)255,	// single quote to rest of characters -- cast to char to avoid compiler sign warning, but intrinsic uses UBYTE, so code is correct
 		0,0,0,0,0,0,0,0,0,0);
 
 	__m128i k_newlines = _mm_set1_epi8('\n');
@@ -1692,7 +1692,7 @@ static const char* scan_to_directive(const char* p, int* p_line_number)
 					*p_line_number = line_number;
 					return p;
 				}
-				else if (pp_char_class[*leading_whitespace_scan] != PP_CHAR_CLASS_whitespace)
+				else if (pp_char_class[(uint8)*leading_whitespace_scan] != PP_CHAR_CLASS_whitespace)
 				{
 					// Not a leading hash!  Continue parsing
 					p++;
@@ -4915,7 +4915,8 @@ static void init_directive(char* s, int hash)
 		exit(1);
 	}
 
-	strcpy_s(directive_hash[hash].name, sizeof(directive_hash[hash].name), s);
+	// name member is sized to fit the largest directive (7 characters plus null) -- use strncpy to avoid warnings
+	strncpy(directive_hash[hash].name, s, sizeof(directive_hash[hash].name));
 	directive_hash[hash].name_len = strlen(s);
 }
 
diff --git a/Engine/Source/Runtime/RenderCore/Private/ShaderCore.cpp b/Engine/Source/Runtime/RenderCore/Private/ShaderCore.cpp
--- a/Engine/Source/Runtime/RenderCore/Private/ShaderCore.cpp
+++ b/Engine/Source/Runtime/RenderCore/Private/ShaderCore.cpp
@@ -2030,7 +2030,9 @@ void ShaderConvertAndStripComments(const FString& ShaderSource, TArray<ANSICHAR>
 		{
 			if (Current[1] == '/')
 			{
+#if PLATFORM_ALWAYS_HAS_SSE4_2
 				SingleLineCommentParse:
+#endif
 				while (!IsEndOfLine(*Current) && Current < End)
 				{
 					++Current;
@@ -2041,7 +2043,9 @@ void ShaderConvertAndStripComments(const FString& ShaderSource, TArray<ANSICHAR>
 				Current += 2;
 				while (Current < End)
 				{
+#if PLATFORM_ALWAYS_HAS_SSE4_2
 					MultiLineCommentParse:
+#endif
 					if (Current[0] == '*' && Current[1] == '/')
 					{
 						Current += 2;
