Shader Preprocessor optimization:  Support for loading include preprocess dependencies in bulk, and drastically reduced string processing, memory allocation, and map overhead.  Roughly 7x faster, saving 12% in low level preprocessor, or 5% overall.

Flattened include dependencies are generated during include scanning at startup, basically for free (perf difference was well below noise).  Bulk dependencies reduce round trips to the shader cache (which require mutex locks), and are indexed by the ANSI text exactly as it appears in the include directive in the source files, allowing a faster case sensitive hash, and avoiding the need for expensive path string operations.  Anything found in a bulk dependency is stored in an array that parallels the dependency array, rather than a map.  Includes stored in IncludeVirtualPathToSharedContentsMap also use an array.

Noting that our string classes (FString) are already case insensitive by default, some unnecessary case conversions were removed.  The separate map of "seen" shaders was also removed, as we can just use the LoadedIncludesCache map for the same purpose.  Where possible, existing ANSI strings are referenced, avoiding dynamic allocation.

#rnx

diff --git a/Engine/Source/Developer/ShaderPreprocessor/Private/ShaderPreprocessor.cpp b/Engine/Source/Developer/ShaderPreprocessor/Private/ShaderPreprocessor.cpp
--- a/Engine/Source/Developer/ShaderPreprocessor/Private/ShaderPreprocessor.cpp
+++ b/Engine/Source/Developer/ShaderPreprocessor/Private/ShaderPreprocessor.cpp
@@ -22,7 +22,6 @@ static TAutoConsoleVariable<int32> CVarShaderCompilerThreadLocalPreprocessBuffer
 namespace
 {
 	const FString PlatformHeader = TEXT("/Engine/Public/Platform.ush");
-	const FString PlatformHeaderLowerCase = PlatformHeader.ToLower();
 	void LogMandatoryHeaderError(const FShaderCompilerInput& Input, FShaderPreprocessOutput& Output)
 	{
 		FString Path = Input.VirtualSourceFilePath;
@@ -31,6 +30,45 @@ namespace
 	}
 }
 
+// Utility function to wrap FShaderPreprocessDependencies hash table lookups -- used with FComparePathInSource / FCompareResultPath below
+template <typename CompareType, typename... ArgsType>
+FORCEINLINE uint32 DependencyHashTableFind(const FShaderPreprocessDependencies& Dependencies, const CompareType& Compare, uint32 KeyHash, ArgsType... Args)
+{
+	const FHashTable& HashTable = Compare.GetHashTable(Dependencies);
+	for (uint32 Index = HashTable.First(KeyHash); HashTable.IsValid(Index); Index = HashTable.Next(Index))
+	{
+		if (Compare.Equals(Dependencies.Dependencies[Index], Args...))
+		{
+			return Index;
+		}
+	}
+	return INDEX_NONE;
+}
+
+struct FComparePathInSource
+{
+	static FORCEINLINE const FHashTable& GetHashTable(const FShaderPreprocessDependencies& Dependencies)
+	{
+		return Dependencies.BySource;
+	}
+	static FORCEINLINE bool Equals(const FShaderPreprocessDependency& Dependency, const ANSICHAR* PathInSource, uint32 PathLen, FXxHash64 PathHash, const ANSICHAR* ParentPathAnsi)
+	{
+		return Dependency.EqualsPathInSource(PathInSource, PathLen, PathHash, ParentPathAnsi);
+	}
+};
+
+struct FCompareResultPath
+{
+	static FORCEINLINE const FHashTable& GetHashTable(const FShaderPreprocessDependencies& Dependencies)
+	{
+		return Dependencies.ByResult;
+	}
+	static FORCEINLINE bool Equals(const FShaderPreprocessDependency& Dependency, const FString& ResultPath, uint32 ResultPathHash)
+	{
+		return Dependency.EqualsResultPath(ResultPath, ResultPathHash);
+	}
+};
+
 PRAGMA_DISABLE_DEPRECATION_WARNINGS		// FShaderCompilerDefinitions will be made internal in the future, marked deprecated until then
 static void AddStbDefine(stb_arena* MacroArena, macro_definition**& StbDefines, const TCHAR* Name, const TCHAR* Value);
 static void AddStbDefines(stb_arena* MacroArena, macro_definition**& StbDefines, const FShaderCompilerDefinitions& Defines);
@@ -101,44 +139,93 @@ extern "C"
 
 struct FStbLoadedInclude
 {
+	const ANSICHAR* FileName = nullptr;			// Points to ResultPath in FShaderPreprocessDependenciesShared, or LocalFileName
 	const ANSICHAR* Data = nullptr;				// Points to SharedData, LocalData, or data from FShaderCompilerEnvironment
 	size_t DataLength = 0;
 	FShaderSharedAnsiStringPtr SharedData;
 	TArray<ANSICHAR> LocalData;
+	TArray<ANSICHAR> LocalFileName;
 };
 
+static bool HasDependencyFromResultPath(const FShaderPreprocessDependencies& Dependencies, const FString& ResultPath, const FStbLoadedInclude* CacheShared);
+
 struct FStbPreprocessContext
 {
 	const FShaderCompilerInput& ShaderInput;
 	const FShaderCompilerEnvironment& Environment;
 	TMap<FString, FStbLoadedInclude> LoadedIncludesCache;
-	TMap<FString, TUniquePtr<ANSICHAR[]>> SeenPathsLowerCase;
+
+	// Shared includes from PreprocessDependencies, VertexFactoryDependencies, and Environment.IncludeVirtualPathToSharedContentsMap
+	// are stored in this array instead of the map, indexed sequentially.  Avoids hash table overhead of "LoadedIncludesCache".
+	TArray<FStbLoadedInclude> LoadedIncludesCacheShared;
+	FShaderPreprocessDependenciesShared PreprocessDependencies;
+	FShaderPreprocessDependenciesShared VertexFactoryDependencies;
+	FHashTable SharedContentsHash;								// Case insensitive hash table pointing at LoadedIncludesCacheShared with entries from IncludeVirtualPathToSharedContentsMap
+	uint32 SharedIncludeIndex = INDEX_NONE;						// Index in LoadedIncludesCacheShared propagated from StbResolveInclude to StbLoadFile
+	uint32 VertexFactoryOffset = INDEX_NONE;					// Vertex factory dependencies start at this offset in LoadedIncludesCacheShared
+	uint32 VirtualSharedContentsOffset = INDEX_NONE;			// Virtual shared contents start at this offset in LoadedIncludesCacheShared
 
 	bool HasIncludedMandatoryHeaders()
 	{
-		return SeenPathsLowerCase.Contains(PlatformHeaderLowerCase);
+		// Checks if the mandatory PlatformHeader ("/Engine/Public/Platform.ush") has been included.  Note that the header may be
+		// encountered through one of our FShaderPreprocessDependencies structures, so if those are valid, we need to check the
+		// corresponding elements in the LoadedIncludesCacheShared array to see if the path was encountered.
+		return
+			(PreprocessDependencies.IsValid() && HasDependencyFromResultPath(*PreprocessDependencies, PlatformHeader, &LoadedIncludesCacheShared[0])) ||
+			(VertexFactoryDependencies.IsValid() && HasDependencyFromResultPath(*VertexFactoryDependencies, PlatformHeader, &LoadedIncludesCacheShared[VertexFactoryOffset])) ||
+			LoadedIncludesCache.Contains(PlatformHeader);
 	}
 };
 
+static void StbLoadedIncludeTrimPaddingChecked(FStbLoadedInclude* ContentsCached)
+{
+	// Need 15 characters beyond null terminator, so an unaligned SSE read at the null terminator can safely read 15 extra unused characters
+	// without going out of memory bounds.  ShaderConvertAndStripComments adds this padding in the form of extra trailing zeroes.  Make sure
+	// these zeroes are there.
+	static const char SixteenZeroes[16] = { 0 };
+	checkf(ContentsCached->DataLength >= 16 && memcmp(&ContentsCached->Data[ContentsCached->DataLength - 16], SixteenZeroes, 16) == 0,
+		TEXT("Shader preprocessor ANSI files must include 15 bytes of zero padding past null terminator"));
+
+	ContentsCached->DataLength -= 15;
+}
+
+static FORCEINLINE void StbLoadedIncludeTrimPadding(FStbLoadedInclude* ContentsCached)
+{
+	// For includes cached at startup, don't bother with the assert, since we know they came from a "safe" source that always adds the padding.
+	ContentsCached->DataLength -= 15;
+}
+
 static const ANSICHAR* StbLoadFile(const ANSICHAR* Filename, void* RawContext, size_t* OutLength)
 {
 	FStbPreprocessContext& Context = *reinterpret_cast<FStbPreprocessContext*>(RawContext);
+
+	// Check if we found this file in our preprocess dependencies (fast path)
+	if (Context.SharedIncludeIndex != INDEX_NONE)
+	{
+		FStbLoadedInclude* ContentsCached = &Context.LoadedIncludesCacheShared[Context.SharedIncludeIndex];
+
+		// Reset this after we consume it (although StbResolveInclude should clear it as well before StbLoadFile is called again)
+		Context.SharedIncludeIndex = INDEX_NONE;
+
+		*OutLength = ContentsCached->DataLength;
+		return ContentsCached->Data;
+	}
+
 	FString FilenameConverted = StringCast<TCHAR>(Filename).Get();
 	uint32 FilenameConvertedHash = GetTypeHash(FilenameConverted);
-	FStbLoadedInclude* ContentsCached = Context.LoadedIncludesCache.FindByHash(FilenameConvertedHash, FilenameConverted);
-	if (!ContentsCached)
-	{
-		ContentsCached = &Context.LoadedIncludesCache.AddByHash(FilenameConvertedHash, FilenameConverted);
 
+	FStbLoadedInclude& ContentsCached = Context.LoadedIncludesCache.FindOrAddByHash(FilenameConvertedHash, FilenameConverted);
+	if (!ContentsCached.Data)
+	{
 		const FString* InMemorySource = Context.Environment.IncludeVirtualPathToContentsMap.FindByHash(FilenameConvertedHash, FilenameConverted);
 
 		if (InMemorySource)
 		{
 			check(!InMemorySource->IsEmpty());
-			ShaderConvertAndStripComments(*InMemorySource, ContentsCached->LocalData);
+			ShaderConvertAndStripComments(*InMemorySource, ContentsCached.LocalData);
 
-			ContentsCached->Data = ContentsCached->LocalData.GetData();
-			ContentsCached->DataLength = ContentsCached->LocalData.Num();
+			ContentsCached.Data = ContentsCached.LocalData.GetData();
+			ContentsCached.DataLength = ContentsCached.LocalData.Num();
 		}
 		else
 		{
@@ -146,31 +233,23 @@ static const ANSICHAR* StbLoadFile(const ANSICHAR* Filename, void* RawContext, s
 
 			if (InMemorySourceAnsi)
 			{
-				ContentsCached->Data = InMemorySourceAnsi->Get()->GetData();
-				ContentsCached->DataLength = InMemorySourceAnsi->Get()->Num();
+				ContentsCached.Data = InMemorySourceAnsi->Get()->GetData();
+				ContentsCached.DataLength = InMemorySourceAnsi->Get()->Num();
 			}
 			else
 			{
 				CheckShaderHashCacheInclude(FilenameConverted, Context.ShaderInput.Target.GetPlatform(), Context.ShaderInput.ShaderFormat.ToString());
-				LoadShaderSourceFile(*FilenameConverted, Context.ShaderInput.Target.GetPlatform(), nullptr, nullptr, nullptr, &ContentsCached->SharedData);
+				LoadShaderSourceFile(*FilenameConverted, Context.ShaderInput.Target.GetPlatform(), nullptr, nullptr, nullptr, &ContentsCached.SharedData);
 
-				ContentsCached->Data = ContentsCached->SharedData->GetData();
-				ContentsCached->DataLength = ContentsCached->SharedData->Num();
+				ContentsCached.Data = ContentsCached.SharedData->GetData();
+				ContentsCached.DataLength = ContentsCached.SharedData->Num();
 			}
 		}
 
-		// Need 15 characters beyond null terminator, so an unaligned SSE read at the null terminator can safely read 15 extra unused characters
-		// without going out of memory bounds.  ShaderConvertAndStripComments adds this padding in the form of extra trailing zeroes.  Make sure
-		// these zeroes are there.
-		static const char SixteenZeroes[16] = { 0 };
-		checkf(ContentsCached->DataLength >= 16 && memcmp(&ContentsCached->Data[ContentsCached->DataLength - 16], SixteenZeroes, 16) == 0,
-			TEXT("Shader preprocessor ANSI files must include 15 bytes of zero padding past null terminator"));
-
-		ContentsCached->DataLength -= 15;
+		StbLoadedIncludeTrimPaddingChecked(&ContentsCached);
 	}
-	check(ContentsCached);
-	*OutLength = ContentsCached->DataLength;
-	return ContentsCached->Data;
+	*OutLength = ContentsCached.DataLength;
+	return ContentsCached.Data;
 }
 
 static void StbFreeFile(const ANSICHAR* Filename, const ANSICHAR* Contents, void* RawContext)
@@ -179,42 +258,181 @@ static void StbFreeFile(const ANSICHAR* Filename, const ANSICHAR* Contents, void
 	// we want to keep it around until that point in case includes are loaded multiple times from different source locations
 }
 
+static uint32 ResolveDependencyFromPathInSource(const FShaderPreprocessDependencies& Dependencies, const ANSICHAR* PathInSource, uint32 PathLen, FXxHash64 PathHash, const ANSICHAR* ParentPathAnsi, FStbLoadedInclude* CacheShared)
+{
+	uint32 HashIndex = DependencyHashTableFind(Dependencies, FComparePathInSource(), GetTypeHash(PathHash), PathInSource, PathLen, PathHash, ParentPathAnsi);
+	if (HashIndex != INDEX_NONE)
+	{
+		// Choose the first unique instance of this result path
+		HashIndex = Dependencies.Dependencies[HashIndex].ResultPathUniqueIndex;
+
+		const FShaderPreprocessDependency& Dependency = Dependencies.Dependencies[HashIndex];
+
+		FStbLoadedInclude* ContentsCached = &CacheShared[HashIndex];
+		if (!ContentsCached->FileName)
+		{
+			ContentsCached->FileName = Dependency.ResultPath.GetData();
+			ContentsCached->Data = Dependency.StrippedSource->GetData();
+			ContentsCached->DataLength = Dependency.StrippedSource->Num();
+			StbLoadedIncludeTrimPadding(ContentsCached);
+		}
+	}
+	return HashIndex;
+}
+
+static uint32 ResolveDependencyFromResultPath(const FShaderPreprocessDependencies& Dependencies, const FString& ResultPath, uint32 ResultPathHash, FStbLoadedInclude* CacheShared)
+{
+	// ResultPathHash is passed in twice -- once for "Find" function, and again as an argument to the "FCompareResultPath::Equals" function
+	uint32 HashIndex = DependencyHashTableFind(Dependencies, FCompareResultPath(), ResultPathHash, ResultPath, ResultPathHash);
+	if (HashIndex != INDEX_NONE)
+	{
+		const FShaderPreprocessDependency& Dependency = Dependencies.Dependencies[HashIndex];
+
+		FStbLoadedInclude* ContentsCached = &CacheShared[HashIndex];
+		if (!ContentsCached->FileName)
+		{
+			ContentsCached->FileName = Dependency.ResultPath.GetData();
+			ContentsCached->Data = Dependency.StrippedSource->GetData();
+			ContentsCached->DataLength = Dependency.StrippedSource->Num();
+			StbLoadedIncludeTrimPadding(ContentsCached);
+		}
+	}
+	return HashIndex;
+}
+
+// Returns true if the path in question was encountered during preprocessing, if the path is one of the paths referenced by that dependency structure.
+static bool HasDependencyFromResultPath(const FShaderPreprocessDependencies& Dependencies, const FString& ResultPath, const FStbLoadedInclude* CacheShared)
+{
+	uint32 ResultPathHash = GetTypeHash(ResultPath);
+	uint32 HashIndex = DependencyHashTableFind(Dependencies, FCompareResultPath(), ResultPathHash, ResultPath, ResultPathHash);
+
+	// Entry will have FileName set if it was encountered
+	return HashIndex != INDEX_NONE && CacheShared[HashIndex].FileName != nullptr;
+}
+
+static void CopyStringToAnsiCharArray(const TCHAR* Text, int32 TextLen, TArray<ANSICHAR>& Out)
+{
+	Out.SetNumUninitialized(TextLen + 1);
+	ANSICHAR* OutData = Out.GetData();
+	for (int32 CharIndex = 0; CharIndex < TextLen; CharIndex++, OutData++, Text++)
+	{
+		*OutData = (ANSICHAR)*Text;
+	}
+	*OutData = 0;
+}
+
 static const ANSICHAR* StbResolveInclude(const ANSICHAR* PathInSource, uint32 PathLen, const ANSICHAR* ParentPathAnsi, void* RawContext)
 {
 	FStbPreprocessContext& Context = *reinterpret_cast<FStbPreprocessContext*>(RawContext);
+
+	FXxHash64 PathHash = FXxHash64::HashBuffer(PathInSource, PathLen);
+
+	// Try main shader preprocess dependencies
+	Context.SharedIncludeIndex = INDEX_NONE;
+	if (Context.PreprocessDependencies.IsValid())
+	{
+		uint32 DependencyIndex = ResolveDependencyFromPathInSource(*Context.PreprocessDependencies, PathInSource, PathLen, PathHash, ParentPathAnsi, &Context.LoadedIncludesCacheShared[0]);
+
+		if (DependencyIndex != INDEX_NONE)
+		{
+			// Propagate the found index to StbLoadFile
+			uint32 SharedIncludeIndex = DependencyIndex;
+			Context.SharedIncludeIndex = SharedIncludeIndex;
+			return Context.LoadedIncludesCacheShared[SharedIncludeIndex].FileName;
+		}
+	}
+
+	// Try vertex factory preprocess dependencies
+	if (Context.VertexFactoryDependencies.IsValid())
+	{
+		uint32 DependencyIndex = ResolveDependencyFromPathInSource(*Context.VertexFactoryDependencies, PathInSource, PathLen, PathHash, ParentPathAnsi, &Context.LoadedIncludesCacheShared[Context.VertexFactoryOffset]);
+
+		if (DependencyIndex != INDEX_NONE)
+		{
+			// Propagate the found index to StbLoadFile
+			uint32 SharedIncludeIndex = DependencyIndex + Context.VertexFactoryOffset;
+			Context.SharedIncludeIndex = SharedIncludeIndex;
+			return Context.LoadedIncludesCacheShared[SharedIncludeIndex].FileName;
+		}
+	}
+
+	// Try SharedContentsHash
+	FAnsiStringView RawPathInSourceView(PathInSource, PathLen);
+	for (uint32 HashIndex = Context.SharedContentsHash.First(GetTypeHash(RawPathInSourceView)); Context.SharedContentsHash.IsValid(HashIndex); HashIndex = Context.SharedContentsHash.Next(HashIndex))
+	{
+		if (RawPathInSourceView == Context.LoadedIncludesCacheShared[HashIndex].FileName)
+		{
+			// Propagate the found index to StbLoadFile
+			Context.SharedIncludeIndex = HashIndex;
+			return Context.LoadedIncludesCacheShared[HashIndex].FileName;
+		}
+	}
+
+	// Slow path...  Platform specific files and procedurally generated files (/Engine/Generated/Material.ush) -- typically 5% of files.
 	FString PathModified(PathLen, PathInSource);
-	FString ParentFolder(ParentPathAnsi);
-	ParentFolder = FPaths::GetPath(ParentFolder);
 	if (!PathModified.StartsWith(TEXT("/"))) // if path doesn't start with / it's relative, if so append the parent's folder and collapse any relative dirs
 	{
+		FString ParentFolder(ParentPathAnsi);
+		ParentFolder = FPaths::GetPath(ParentFolder);
 		PathModified = ParentFolder / PathModified;
 		FPaths::CollapseRelativeDirectories(PathModified);
 	}
 
 	FixupShaderFilePath(PathModified, Context.ShaderInput.Target.GetPlatform(), &Context.ShaderInput.ShaderPlatformName);
+	uint32 PathModifiedHash = GetTypeHash(PathModified);
+
+	// We need to check our preprocess dependencies again with the result path, so we get the canonical capitalization for it from the dependencies, if available.
+	// This case can be reached for platform includes (which aren't added to the bulk dependencies).
+	if (Context.PreprocessDependencies.IsValid())
+	{
+		uint32 DependencyIndex = ResolveDependencyFromResultPath(*Context.PreprocessDependencies, PathModified, PathModifiedHash, &Context.LoadedIncludesCacheShared[0]);
+
+		if (DependencyIndex != INDEX_NONE)
+		{
+			// Propagate the found index to StbLoadFile
+			uint32 SharedIncludeIndex = DependencyIndex;
+			Context.SharedIncludeIndex = SharedIncludeIndex;
+			return Context.LoadedIncludesCacheShared[SharedIncludeIndex].FileName;
+		}
+	}
+
+	// Try vertex factory preprocess dependencies
+	if (Context.VertexFactoryDependencies.IsValid())
+	{
+		uint32 DependencyIndex = ResolveDependencyFromResultPath(*Context.VertexFactoryDependencies, PathModified, PathModifiedHash, &Context.LoadedIncludesCacheShared[Context.VertexFactoryOffset]);
+
+		if (DependencyIndex != INDEX_NONE)
+		{
+			// Propagate the found index to StbLoadFile
+			uint32 SharedIncludeIndex = DependencyIndex + Context.VertexFactoryOffset;
+			Context.SharedIncludeIndex = SharedIncludeIndex;
+			return Context.LoadedIncludesCacheShared[SharedIncludeIndex].FileName;
+		}
+	}
 
-	FString PathModifiedLowerCase = PathModified.ToLower();
-	const TUniquePtr<ANSICHAR[]>* SeenPath = Context.SeenPathsLowerCase.Find(PathModifiedLowerCase);
-	// Keep track of previously resolved paths in a case insensitive manner so preprocessor will handle #pragma once with files included with inconsistent casing correctly
-	// (we store the first correctly resolved path with original casing so we get "nice" line directives)
-	if (SeenPath)
+	// If we reach here, the include will be added to the map.  Check if it's already in the map.
+	FStbLoadedInclude* ContentsCached = Context.LoadedIncludesCache.FindByHash(PathModifiedHash, PathModified);
+	if (ContentsCached)
 	{
-		return SeenPath->Get();
+		// We return the same previously resolved path so preprocessor will handle #pragma once with files included with inconsistent casing correctly
+		return ContentsCached->FileName;
 	}
 
 	bool bExists =
-		Context.Environment.IncludeVirtualPathToContentsMap.Contains(PathModified) ||
-		Context.Environment.IncludeVirtualPathToSharedContentsMap.Contains(PathModified) ||
+		Context.Environment.IncludeVirtualPathToContentsMap.ContainsByHash(PathModifiedHash, PathModified) ||
 		// LoadShaderSourceFile will load the file if it exists, but then cache it internally, so the next call in StbLoadFile will be cheap
 		// (and hence this is not wasteful, just performs the loading earlier)
 		LoadShaderSourceFile(*PathModified, Context.ShaderInput.Target.GetPlatform(), nullptr, nullptr);
 
 	if (bExists)
 	{
-		int32 Length = FPlatformString::ConvertedLength<ANSICHAR>(*PathModified);
-		TUniquePtr<ANSICHAR[]>& OutPath = Context.SeenPathsLowerCase.Add(PathModifiedLowerCase, MakeUnique<ANSICHAR[]>(Length));
-		FPlatformString::Convert<TCHAR, ANSICHAR>(OutPath.Get(), Length, *PathModified);
-		return OutPath.Get();
+		ContentsCached = &Context.LoadedIncludesCache.AddByHash(PathModifiedHash, PathModified);
+
+		// Initialize the ANSI file name in the map entry.  The file itself will be loaded in StbLoadFile, but we need the ANSI string
+		// as the return value from this function. 
+		CopyStringToAnsiCharArray(&PathModified[0], PathModified.Len(), ContentsCached->LocalFileName);
+		ContentsCached->FileName = ContentsCached->LocalFileName.GetData();
+		return ContentsCached->FileName;
 	}
 
 	return nullptr;
@@ -285,6 +503,65 @@ bool InnerPreprocessShaderStb(
 		ThreadLocalPreprocessBuffer = new char[ClampedPreprocessBufferSize];
 	}
 
+	if (GetShaderPreprocessDependencies(*Input.VirtualSourceFilePath, Context.ShaderInput.Target.GetPlatform(), Context.PreprocessDependencies))
+	{
+		// First item in dependencies is always root file, so set that index
+		Context.SharedIncludeIndex = 0;
+	}
+
+	// Grab vertex factory dependencies if present
+	const FString* VertexFactoryInclude = Context.Environment.IncludeVirtualPathToContentsMap.Find(TEXT("/Engine/Generated/VertexFactory.ush"));
+	if (VertexFactoryInclude)
+	{
+		int32 VertexFactoryNameStart;
+		int32 VertexFactoryNameEnd;
+		if (VertexFactoryInclude->FindChar(TEXT('\"'), VertexFactoryNameStart) && VertexFactoryInclude->FindLastChar(TEXT('\"'), VertexFactoryNameEnd))
+		{
+			// Should have at least one character in our filename
+			check(VertexFactoryNameEnd > VertexFactoryNameStart + 1);
+			FString VertexFactoryFilename(FStringView(&(*VertexFactoryInclude)[VertexFactoryNameStart + 1], VertexFactoryNameEnd - (VertexFactoryNameStart + 1)));
+			GetShaderPreprocessDependencies(*VertexFactoryFilename, Context.ShaderInput.Target.GetPlatform(), Context.VertexFactoryDependencies);
+		}
+	}
+
+	// Initialize array of loaded includes associated with PreprocessDependencies, VertexFactoryDependencies, and Environment.IncludeVirtualPathToSharedContentsMap
+	Context.VertexFactoryOffset = Context.PreprocessDependencies.IsValid() ? Context.PreprocessDependencies->Dependencies.Num() : 0;
+	Context.VirtualSharedContentsOffset = Context.VertexFactoryOffset + (Context.VertexFactoryDependencies.IsValid() ? Context.VertexFactoryDependencies->Dependencies.Num() : 0);
+	Context.LoadedIncludesCacheShared.AddDefaulted(Context.VirtualSharedContentsOffset + Context.Environment.IncludeVirtualPathToSharedContentsMap.Num());
+
+	// Initialize root file dependency, if present
+	if (Context.PreprocessDependencies.IsValid())
+	{
+		const FShaderPreprocessDependency& Dependency = Context.PreprocessDependencies->Dependencies[0];
+		FStbLoadedInclude* ContentsCached = &Context.LoadedIncludesCacheShared[0];
+
+		ContentsCached->FileName = InFilename.Get();
+		ContentsCached->Data = Dependency.StrippedSource->GetData();
+		ContentsCached->DataLength = Dependency.StrippedSource->Num();
+		StbLoadedIncludeTrimPadding(ContentsCached);
+	}
+
+	// Initialize loaded includes for IncludeVirtualPathToSharedContentsMap, and generate a hash table
+	uint32 SharedContentsMapIndex = Context.VirtualSharedContentsOffset;
+	for (const auto& SharedContentsMapIt : Context.Environment.IncludeVirtualPathToSharedContentsMap)
+	{
+		FStbLoadedInclude& Include = Context.LoadedIncludesCacheShared[SharedContentsMapIndex];
+
+		// Copy name
+		CopyStringToAnsiCharArray(&SharedContentsMapIt.Key[0], SharedContentsMapIt.Key.Len(), Include.LocalFileName);
+		Include.FileName = Include.LocalFileName.GetData();
+
+		// Set data
+		Include.Data = SharedContentsMapIt.Value->GetData();
+		Include.DataLength = SharedContentsMapIt.Value->Num();
+		StbLoadedIncludeTrimPadding(&Include);
+
+		// Add to hash table -- GetTypeHash on string view is case insensitive
+		Context.SharedContentsHash.Add(GetTypeHash(FAnsiStringView(Include.LocalFileName.GetData(), Include.LocalFileName.Num() - 1)), SharedContentsMapIndex);
+
+		SharedContentsMapIndex++;
+	}
+
 	char* OutPreprocessedAnsi = preprocess_file(InFilename.Get(), &Context, StbDefines, arrlen(StbDefines), &Diagnostics, &NumDiagnostics, ThreadLocalPreprocessBuffer, ClampedPreprocessBufferSize);
 
 	bool HasError = false;
diff --git a/Engine/Source/Runtime/RenderCore/Private/ShaderCore.cpp b/Engine/Source/Runtime/RenderCore/Private/ShaderCore.cpp
--- a/Engine/Source/Runtime/RenderCore/Private/ShaderCore.cpp
+++ b/Engine/Source/Runtime/RenderCore/Private/ShaderCore.cpp
@@ -166,7 +166,8 @@ constexpr int32 GSHADERFILECACHE_BUCKETS = 31; /* prime number for best distribu
 struct FShaderFileCacheEntry
 {
 	FString Source;
-	FShaderSharedAnsiStringPtr StrippedSource;		// Source with comments stripped out, and converted to ANSICHAR
+	FShaderSharedAnsiStringPtr StrippedSource;				// Source with comments stripped out, and converted to ANSICHAR
+	FShaderPreprocessDependenciesShared Dependencies;		// Stripped source with include dependencies, all in one shareable struct
 };
 
 struct FShaderFileCache
@@ -267,12 +268,17 @@ public:
 		return Platforms[ShaderPlatform].ShaderHashCache.Add(VirtualFilePath, FSHAHash());
 	}
 
+	static bool IsPlatformInclude(const FString& VirtualFilePath)
+	{
+		return (VirtualFilePath.StartsWith(TEXT("/Engine/Private/Platform/"))
+			|| VirtualFilePath.StartsWith(TEXT("/Engine/Public/Platform/"))
+			|| VirtualFilePath.StartsWith(TEXT("/Platform/")));
+	}
+
 	bool ShouldIgnoreInclude(const FString& VirtualFilePath, EShaderPlatform ShaderPlatform) const
 	{
 		// Ignore only platform specific files, which won't be used by the target platform.
-		if (VirtualFilePath.StartsWith(TEXT("/Engine/Private/Platform/"))
-			|| VirtualFilePath.StartsWith(TEXT("/Engine/Public/Platform/"))
-			|| VirtualFilePath.StartsWith(TEXT("/Platform/")))
+		if (IsPlatformInclude(VirtualFilePath))
 		{
 			const FString& PlatformIncludeDirectory = GetPlatformIncludeDirectory(ShaderPlatform);
 			if (PlatformIncludeDirectory.IsEmpty() || !(VirtualFilePath.Contains(PlatformIncludeDirectory)))
@@ -1533,6 +1539,8 @@ void CompileShaderPipeline(const TArray<const IShaderFormat*>& ShaderFormats, FS
 	PipelineJob->bSucceeded = true;
 }
 
+static void InternalGetShaderIncludes(const TCHAR* EntryPointVirtualFilePath, const TCHAR* VirtualFilePath, TArray<FString>& IncludeVirtualFilePaths, EShaderPlatform ShaderPlatform, uint32 DepthLimit, bool AddToIncludeFile, const FName* ShaderPlatformName, bool bPreprocessDependencies);
+
 /**
 * Add a new entry to the list of shader source files
 * Only unique entries which can be loaded are added as well as their #include files
@@ -1550,7 +1558,8 @@ void AddShaderSourceFileEntry(TArray<FString>& OutVirtualFilePaths, FString Virt
 		TArray<FString> ShaderIncludes;
 
 		const uint32 DepthLimit = 100;
-		GetShaderIncludes(*VirtualFilePath, *VirtualFilePath, OutVirtualFilePaths, ShaderPlatform, DepthLimit, ShaderPlatformName);
+		const bool bPreprocessDependencies = true;
+		InternalGetShaderIncludes(*VirtualFilePath, *VirtualFilePath, OutVirtualFilePaths, ShaderPlatform, DepthLimit, false, ShaderPlatformName, bPreprocessDependencies);
 		for( int32 IncludeIdx=0; IncludeIdx < ShaderIncludes.Num(); IncludeIdx++ )
 		{
 			OutVirtualFilePaths.AddUnique(ShaderIncludes[IncludeIdx]);
@@ -2238,10 +2247,102 @@ static const TCHAR* FindFirstInclude(const TCHAR* Text)
 	return nullptr;
 }
 
+static void StringCopyToAnsiCharArray(const TCHAR* Text, int32 TextLen, TArray<ANSICHAR>& Out)
+{
+	Out.SetNumUninitialized(TextLen + 1);
+	ANSICHAR* OutData = Out.GetData();
+	for (int32 CharIndex = 0; CharIndex < TextLen; CharIndex++, OutData++, Text++)
+	{
+		*OutData = (ANSICHAR)*Text;
+	}
+	*OutData = 0;
+}
+
+// Allocates structure and adds root file dependency
+static FShaderPreprocessDependencies* ShaderPreprocessDependenciesBegin(const TCHAR* VirtualFilePath)
+{
+	FShaderPreprocessDependencies* PreprocessDependencies = new FShaderPreprocessDependencies();
+
+	PreprocessDependencies->Dependencies.AddDefaulted();
+	StringCopyToAnsiCharArray(VirtualFilePath, FCString::Strlen(VirtualFilePath), PreprocessDependencies->Dependencies[0].ResultPath);
+	PreprocessDependencies->Dependencies[0].ResultPathHash = FCrc::Strihash_DEPRECATED(VirtualFilePath);
+
+	return PreprocessDependencies;
+}
+
+// Adds finished dependencies to the cache
+static void ShaderPreprocessDependenciesEnd(const TCHAR* VirtualFilePath, FShaderPreprocessDependencies* PreprocessDependencies, EShaderPlatform Platform)
+{
+	uint32 CurrentHash = FCrc::Strihash_DEPRECATED(VirtualFilePath);
+	FShaderFileCache& ShaderFileCache = GShaderFileCache[CurrentHash % GSHADERFILECACHE_BUCKETS];
+	{
+		FRWScopeLock ScopeLock(ShaderFileCache.Lock, SLT_Write);
+		FShaderFileCacheEntry* CachedFile = ShaderFileCache.Map.FindByHash(CurrentHash, VirtualFilePath);
+		if (CachedFile)
+		{
+			// Another thread could have finished the job...  If not, set the dependencies.
+			if (!CachedFile->Dependencies.IsValid())
+			{
+				CachedFile->Dependencies = MakeShareable(PreprocessDependencies);
+			}
+			else
+			{
+				delete PreprocessDependencies;
+			}
+		}
+	}
+}
+
+static void AddPreprocessDependency(FShaderPreprocessDependencies& Dependencies, const FShaderPreprocessDependency& Dependency)
+{
+	check(Dependency.StrippedSource.IsValid());
+
+	// First, check if the dependency already exists
+	for (uint32 HashIndex = Dependencies.BySource.First(GetTypeHash(Dependency.PathInSourceHash)); Dependencies.BySource.IsValid(HashIndex); HashIndex = Dependencies.BySource.Next(HashIndex))
+	{
+		FShaderPreprocessDependency& TestDependency = Dependencies.Dependencies[HashIndex];
+
+		// Subtract one from PathInSource.Num() to get length minus null terminator
+		if (TestDependency.EqualsPathInSource(Dependency.PathInSource.GetData(), Dependency.PathInSource.Num() - 1, Dependency.PathInSourceHash, Dependency.ParentPath.GetData()))
+		{
+			// The result path better be the same for both
+			check(!FCStringAnsi::Stricmp(TestDependency.ResultPath.GetData(), Dependency.ResultPath.GetData()));
+			return;
+		}
+	}
+
+	// Add the dependency
+	int32 AddedIndex = Dependencies.Dependencies.Add(Dependency);
+	Dependencies.BySource.Add(GetTypeHash(Dependency.PathInSourceHash), (uint32)AddedIndex);
+
+	// Then check if the result path already exists, so we can point ResultPathUniqueIndex at the first instance of the result path
+	uint32 ExistingResultIndex;
+	for (ExistingResultIndex = Dependencies.ByResult.First(Dependency.ResultPathHash); Dependencies.ByResult.IsValid(ExistingResultIndex); ExistingResultIndex = Dependencies.ByResult.Next(ExistingResultIndex))
+	{
+		FShaderPreprocessDependency& TestDependency = Dependencies.Dependencies[ExistingResultIndex];
+		if (TestDependency.EqualsResultPath(Dependency.ResultPath.GetData(), Dependency.ResultPathHash))
+		{
+			break;
+		}
+	}
+
+	if (Dependencies.ByResult.IsValid(ExistingResultIndex))
+	{
+		// Reference existing result
+		Dependencies.Dependencies[AddedIndex].ResultPathUniqueIndex = ExistingResultIndex;
+	}
+	else
+	{
+		// Add new result
+		Dependencies.Dependencies[AddedIndex].ResultPathUniqueIndex = (uint32)AddedIndex;
+		Dependencies.ByResult.Add(Dependency.ResultPathHash, (uint32)AddedIndex);
+	}
+}
+
 /**
  * Recursively populates IncludeFilenames with the unique include filenames found in the shader file named Filename.
  */
-static void InternalGetShaderIncludes(const TCHAR* EntryPointVirtualFilePath, const TCHAR* VirtualFilePath, const FString& FileContents, TArray<FString>& IncludeVirtualFilePaths, EShaderPlatform ShaderPlatform, uint32 DepthLimit, bool AddToIncludeFile, const FName* ShaderPlatformName)
+static void InternalGetShaderIncludes(const TCHAR* EntryPointVirtualFilePath, const TCHAR* VirtualFilePath, const FString& FileContents, TArray<FString>& IncludeVirtualFilePaths, EShaderPlatform ShaderPlatform, uint32 DepthLimit, bool AddToIncludeFile, const FName* ShaderPlatformName, FShaderPreprocessDependencies* OutDependencies)
 {
 	//avoid an infinite loop with a 0 length string
 	if (FileContents.Len() > 0)
@@ -2282,15 +2383,20 @@ static void InternalGetShaderIncludes(const TCHAR* EntryPointVirtualFilePath, co
 					}
 
 					//CRC the template, not the filled out version so that this shader's CRC will be independent of which material references it.
-					if (ExtractedIncludeFilename == TEXT("/Engine/Generated/Material.ush"))
+					const TCHAR* MaterialTemplateName = TEXT("/Engine/Private/MaterialTemplate.ush");
+					const TCHAR* MaterialGeneratedName = TEXT("/Engine/Generated/Material.ush");
+
+					bool bIsMaterialTemplate = false;
+					if (ExtractedIncludeFilename == MaterialGeneratedName)
 					{
-						ExtractedIncludeFilename = TEXT("/Engine/Private/MaterialTemplate.ush");
+						ExtractedIncludeFilename = MaterialTemplateName;
+						bIsMaterialTemplate = true;
 					}
 
-					ReplaceVirtualFilePathForShaderPlatform(ExtractedIncludeFilename, ShaderPlatform);
+					bool bIsPlatformFile = ReplaceVirtualFilePathForShaderPlatform(ExtractedIncludeFilename, ShaderPlatform);
 
 					// Fixup autogen file
-					ReplaceVirtualFilePathForShaderAutogen(ExtractedIncludeFilename, ShaderPlatform, ShaderPlatformName);
+					bIsPlatformFile |= ReplaceVirtualFilePathForShaderAutogen(ExtractedIncludeFilename, ShaderPlatform, ShaderPlatformName);
 
 					// Ignore uniform buffer, vertex factory and instanced stereo includes
 					bool bIgnoreInclude = ExtractedIncludeFilename.StartsWith(TEXT("/Engine/Generated/"));
@@ -2304,17 +2410,92 @@ static void InternalGetShaderIncludes(const TCHAR* EntryPointVirtualFilePath, co
 						bIgnoreInclude = bIgnoreInclude || GShaderHashCache.ShouldIgnoreInclude(ExtractedIncludeFilename, ShaderPlatform);
 					}
 
+					bIsPlatformFile |= FShaderHashCache::IsPlatformInclude(ExtractedIncludeFilename);
 
 					//vertex factories need to be handled separately
 					if (!bIgnoreInclude)
 					{
-						if (!IncludeVirtualFilePaths.Contains(ExtractedIncludeFilename))
+						int32 SeenFilenameIndex = IncludeVirtualFilePaths.Find(ExtractedIncludeFilename);
+						if (SeenFilenameIndex == INDEX_NONE)
 						{
+							// Preprocess dependencies don't include platform files.
+							FShaderPreprocessDependencies* ExtractedIncludeDependencies = nullptr;
+							if (OutDependencies && !bIsPlatformFile)
+							{
+								ExtractedIncludeDependencies = ShaderPreprocessDependenciesBegin(*ExtractedIncludeFilename);
+							}
+
+							// First element in Dependencies is root file, so initialize the StrippedSource pointer in it
 							FString IncludedFileContents;
-							LoadShaderSourceFile(*ExtractedIncludeFilename, ShaderPlatform, &IncludedFileContents, nullptr, ShaderPlatformName);
-							InternalGetShaderIncludes(EntryPointVirtualFilePath, *ExtractedIncludeFilename, IncludedFileContents, IncludeVirtualFilePaths, ShaderPlatform, DepthLimit - 1, true, ShaderPlatformName);
+							LoadShaderSourceFile(*ExtractedIncludeFilename, ShaderPlatform, &IncludedFileContents, nullptr, ShaderPlatformName,
+								ExtractedIncludeDependencies ? &ExtractedIncludeDependencies->Dependencies[0].StrippedSource : nullptr);
+
+							InternalGetShaderIncludes(EntryPointVirtualFilePath, *ExtractedIncludeFilename, IncludedFileContents, IncludeVirtualFilePaths, ShaderPlatform, DepthLimit - 1, true, ShaderPlatformName, ExtractedIncludeDependencies);
+
+							if (ExtractedIncludeDependencies)
+							{
+								// Some generated shaders are referenced as includes, and won't be found -- if so, just delete the dependencies
+								if (ExtractedIncludeDependencies->Dependencies[0].StrippedSource.IsValid())
+								{
+									ShaderPreprocessDependenciesEnd(*ExtractedIncludeFilename, ExtractedIncludeDependencies, ShaderPlatform);
+								}
+								else
+								{
+									delete ExtractedIncludeDependencies;
+								}
+							}
 						}
-					}
+
+						if (OutDependencies)
+						{
+							// Preprocess dependencies don't include platform files.
+							if (!bIsPlatformFile)
+							{
+								// The material template itself isn't added as a dependency, but child includes of it are.
+								FShaderSharedAnsiStringPtr StrippedContents;
+								if (!bIsMaterialTemplate && LoadShaderSourceFile(*ExtractedIncludeFilename, ShaderPlatform, nullptr, nullptr, nullptr, &StrippedContents))
+								{
+									// Add immediate dependency
+									FShaderPreprocessDependency Dependency;
+									Dependency.StrippedSource = StrippedContents;
+
+									// If the parent is the material template, switch its name to the generated name, so include dependencies from
+									// the material template to other non-procedural files can be cached.
+									const TCHAR* ParentNonTemplate = VirtualFilePath == MaterialTemplateName ? MaterialGeneratedName : VirtualFilePath;
+
+									// We want ResultPath to have consistent case, for the preprocessor which is case sensitive.  So we use the exact
+									// string from the previously found array element if it exists.  If this is the first time it's encountered, it will
+									// have been added to the array by the InternalGetShaderIncludes call above.
+									const FString& ResultPath = SeenFilenameIndex == INDEX_NONE ? ExtractedIncludeFilename : IncludeVirtualFilePaths[SeenFilenameIndex];
+
+									StringCopyToAnsiCharArray(IncludeFilenameBegin + 1, (int32)(IncludeFilenameEnd - IncludeFilenameBegin - 1), Dependency.PathInSource);
+									StringCopyToAnsiCharArray(ParentNonTemplate, FCString::Strlen(ParentNonTemplate), Dependency.ParentPath);
+									StringCopyToAnsiCharArray(*ResultPath, ResultPath.Len(), Dependency.ResultPath);
+									Dependency.ResultPathHash = GetTypeHash(ResultPath);
+
+									// Hash deliberately doesn't include null terminator, so we can generate hash from string view.  Xxhash is faster than
+									// the normal case insensitive string hash, so we choose that.
+									Dependency.PathInSourceHash = FXxHash64::HashBuffer(Dependency.PathInSource.GetData(), Dependency.PathInSource.Num() - 1);
+
+									AddPreprocessDependency(*OutDependencies, Dependency);
+								}
+
+								// Add recursive dependencies from the child
+								FShaderPreprocessDependenciesShared ChildDependenciesShared;
+								if (GetShaderPreprocessDependencies(*ExtractedIncludeFilename, ShaderPlatform, ChildDependenciesShared))
+								{
+									const FShaderPreprocessDependencies& ChildDependencies = *ChildDependenciesShared;
+
+									// Skip over first entry, which is the root file (its dependency is handled by the "add immediate dependency" code above)
+									for (int32 DependencyIndex = 1; DependencyIndex < ChildDependencies.Dependencies.Num(); DependencyIndex++)
+									{
+										AddPreprocessDependency(*OutDependencies, ChildDependencies.Dependencies[DependencyIndex]);
+									}
+								}
+
+							}  // if (!bIsPlatformFile)
+						}  // if (OutDependencies)
+					}  // if (!bIgnoreInclude)
 				}
 			}
 
@@ -2340,25 +2521,62 @@ static void InternalGetShaderIncludes(const TCHAR* EntryPointVirtualFilePath, co
 	}
 }
 
+bool GetShaderPreprocessDependencies(const TCHAR* VirtualFilePath, EShaderPlatform ShaderPlatform, FShaderPreprocessDependenciesShared& OutDependencies)
+{
+	// Same case insensitive hash used by FString
+	uint32 CurrentHash = FCrc::Strihash_DEPRECATED(VirtualFilePath);
+	FShaderFileCache& ShaderFileCache = GShaderFileCache[CurrentHash % GSHADERFILECACHE_BUCKETS];
+
+	FRWScopeLock ScopeLock(ShaderFileCache.Lock, SLT_ReadOnly);
+	FShaderFileCacheEntry* CachedFile = ShaderFileCache.Map.FindByHash(CurrentHash, VirtualFilePath);
+	if (CachedFile && CachedFile->Dependencies.IsValid())
+	{
+		OutDependencies = CachedFile->Dependencies;
+		return true;
+	}
+
+	return false;
+}
+
 /**
  * Recursively populates IncludeFilenames with the unique include filenames found in the shader file named Filename.
  */
-static void InternalGetShaderIncludes(const TCHAR* EntryPointVirtualFilePath, const TCHAR* VirtualFilePath, TArray<FString>& IncludeVirtualFilePaths, EShaderPlatform ShaderPlatform, uint32 DepthLimit, bool AddToIncludeFile, const FName* ShaderPlatformName)
+static void InternalGetShaderIncludes(const TCHAR* EntryPointVirtualFilePath, const TCHAR* VirtualFilePath, TArray<FString>& IncludeVirtualFilePaths, EShaderPlatform ShaderPlatform, uint32 DepthLimit, bool AddToIncludeFile, const FName* ShaderPlatformName, bool bPreprocessDependencies)
 {
+	FShaderPreprocessDependencies* PreprocessDependencies = nullptr;
+	if (bPreprocessDependencies)
+	{
+		// Check if they've already been generated.  These are platform independent, so we only need to generate them once if multiple platforms are being cooked,
+		// but in case we want to specialize them by platform in the future, the platform is passed in.
+		FShaderPreprocessDependenciesShared OutDependenciesIgnored;
+		if (!GetShaderPreprocessDependencies(VirtualFilePath, ShaderPlatform, OutDependenciesIgnored))
+		{
+			// Allocates dependency structure and adds root file element
+			PreprocessDependencies = ShaderPreprocessDependenciesBegin(VirtualFilePath);
+		}
+	}
+
+	// First element in Dependencies is root file, so initialize the StrippedSource pointer in it
 	FString FileContents;
-	LoadShaderSourceFile(VirtualFilePath, ShaderPlatform, &FileContents, nullptr, ShaderPlatformName);
+	LoadShaderSourceFile(VirtualFilePath, ShaderPlatform, &FileContents, nullptr, ShaderPlatformName, PreprocessDependencies ? &PreprocessDependencies->Dependencies[0].StrippedSource : nullptr);
+
+	InternalGetShaderIncludes(EntryPointVirtualFilePath, VirtualFilePath, FileContents, IncludeVirtualFilePaths, ShaderPlatform, DepthLimit, AddToIncludeFile, ShaderPlatformName, PreprocessDependencies);
 
-	InternalGetShaderIncludes(EntryPointVirtualFilePath, VirtualFilePath, FileContents, IncludeVirtualFilePaths, ShaderPlatform, DepthLimit, AddToIncludeFile, ShaderPlatformName);
+	if (PreprocessDependencies)
+	{
+		// Adds completed dependency structure to shader cache map entry
+		ShaderPreprocessDependenciesEnd(VirtualFilePath, PreprocessDependencies, ShaderPlatform);
+	}
 }
 
 void GetShaderIncludes(const TCHAR* EntryPointVirtualFilePath, const TCHAR* VirtualFilePath, TArray<FString>& IncludeVirtualFilePaths, EShaderPlatform ShaderPlatform, uint32 DepthLimit, const FName* ShaderPlatformName)
 {
-	InternalGetShaderIncludes(EntryPointVirtualFilePath, VirtualFilePath, IncludeVirtualFilePaths, ShaderPlatform, DepthLimit, false, ShaderPlatformName);
+	InternalGetShaderIncludes(EntryPointVirtualFilePath, VirtualFilePath, IncludeVirtualFilePaths, ShaderPlatform, DepthLimit, false, ShaderPlatformName, false);
 }
 
 void GetShaderIncludes(const TCHAR* EntryPointVirtualFilePath, const TCHAR* VirtualFilePath, const FString& FileContents, TArray<FString>& IncludeVirtualFilePaths, EShaderPlatform ShaderPlatform, uint32 DepthLimit, const FName* ShaderPlatformName)
 {
-	InternalGetShaderIncludes(EntryPointVirtualFilePath, VirtualFilePath, FileContents, IncludeVirtualFilePaths, ShaderPlatform, DepthLimit, false, ShaderPlatformName);
+	InternalGetShaderIncludes(EntryPointVirtualFilePath, VirtualFilePath, FileContents, IncludeVirtualFilePaths, ShaderPlatform, DepthLimit, false, ShaderPlatformName, nullptr);
 }
 
 void HashShaderFileWithIncludes(FArchive& HashingArchive, const TCHAR* VirtualFilePath, const FString& FileContents, EShaderPlatform ShaderPlatform, bool bOnlyHashIncludedFiles)
diff --git a/Engine/Source/Runtime/RenderCore/Public/ShaderCompilerCore.h b/Engine/Source/Runtime/RenderCore/Public/ShaderCompilerCore.h
--- a/Engine/Source/Runtime/RenderCore/Public/ShaderCompilerCore.h
+++ b/Engine/Source/Runtime/RenderCore/Public/ShaderCompilerCore.h
@@ -9,6 +9,7 @@
 #include "CoreMinimal.h"
 #include "HAL/FileManager.h"
 #include "Hash/Blake3.h"
+#include "Hash/xxhash.h"
 #include "Stats/Stats.h"
 #include "Templates/RefCounting.h"
 #include "Misc/EnumClassFlags.h"
@@ -771,6 +772,82 @@ extern RENDERCORE_API void ShaderConvertAndStripComments(const FString& ShaderSo
  */
 extern RENDERCORE_API bool LoadShaderSourceFile(const TCHAR* VirtualFilePath, EShaderPlatform ShaderPlatform, FString* OutFileContents, TArray<FShaderCompilerError>* OutCompileErrors, const FName* ShaderPlatformName = nullptr, FShaderSharedAnsiStringPtr* OutStrippedContents = nullptr);
 
+
+struct FShaderPreprocessDependency
+{
+	FXxHash64					PathInSourceHash;		// PathInSourceHash doesn't include PathInSource's null terminator, so hash computation can use a string view
+	TArray<ANSICHAR>			PathInSource;			// Path as it appears in include directive in original shader source, allowing faster case sensitive hash
+	TArray<ANSICHAR>			ParentPath;				// For relative paths, ResultPath is dependent on the parent file the include directive is found in
+	TArray<ANSICHAR>			ResultPath;
+	uint32						ResultPathHash;			// Case insensitive hash of ResultPath (compatible with hash of corresponding FString)
+	uint32						ResultPathUniqueIndex;	// Index of first instance of a given result path in Dependencies array
+	FShaderSharedAnsiStringPtr	StrippedSource;			// Source with comments stripped out, and converted to ANSICHAR (output of ShaderConvertAndStripComments)
+
+	FORCEINLINE bool EqualsPathInSource(const ANSICHAR* InPathInSource, int32 InPathInSourceLen, FXxHash64 InPathInSourceHash, const ANSICHAR* InParentPath) const
+	{
+		// PathInSource is case sensitive, ParentPath is case insensitive.
+		// If the path is absolute (starts with '/'), then the parent path isn't relevant, and shouldn't be checked.
+		return
+			PathInSourceHash == InPathInSourceHash &&
+			(PathInSource[0] == '/' || !FCStringAnsi::Stricmp(ParentPath.GetData(), InParentPath)) &&
+			!FCStringAnsi::Strncmp(PathInSource.GetData(), InPathInSource, InPathInSourceLen);
+	}
+
+	FORCEINLINE bool EqualsResultPath(const FString& InResultPath, uint32 InResultPathHash) const
+	{
+		return (ResultPathHash == InResultPathHash) && InResultPath.Equals(ResultPath.GetData(), ESearchCase::IgnoreCase);
+	}
+
+	FORCEINLINE bool EqualsResultPath(const ANSICHAR* InResultPath, uint32 InResultPathHash) const
+	{
+		return (ResultPathHash == InResultPathHash) && !FCStringAnsi::Stricmp(ResultPath.GetData(), InResultPath);
+	}
+};
+
+
+// Structure that provides an array of #include dependencies for a given root shader file, including not just immediate
+// dependencies, but recursive dependencies from children as well.  Not exhaustive, as it does not include platform
+// specific or generated files, although it does include children of "/Engine/Generated/Material.ush", as derived from
+// "/Engine/Private/MaterialTemplate.ush".  Take the example of ClearUAV.usf:
+//
+// /Engine/Private/Tools/ClearUAV.usf    #include "../Common.ush"
+// /Engine/Private/Common.ush            #include "/Engine/Public/Platform.ush"
+//                                       #include "PackUnpack.ush"
+// /Engine/Public/Platform.ush           #include "FP16Math.ush"
+//
+// The above is a small subset, but the above (and many more) would all show up as elements in Dependencies:
+//
+//      PathInSource                   ParentPath                                ResultPath
+//      --------------------------------------------------------------------------------------------------------
+//      ../Common.ush                  /Engine/Private/Tools/ClearUAV.usf        /Engine/Private/Common.ush
+//      /Engine/Public/Platform.ush    /Engine/Private/Common.ush                /Engine/Public/Platform.ush
+//      PackUnpack.ush                 /Engine/Private/Common.ush                /Engine/Private/PackUnpack.ush
+//      FP16Math.ush                   /Engine/Public/Platform.ush               /Engine/Public/FP16Math.ush
+//
+// The goal of this structure is to allow a shader preprocessor implementation to fetch most of the source dependencies in a
+// single query of the loaded shader cache, and then efficiently search for dependencies encountered in the shader source
+// code, without needing to do string operations to resolve paths (such as converting relative paths like "../Common.ush" to
+// "/Engine/Private/Common.ush").  Besides that, the array organization can be used to manage encountered source files
+// by index, rather than needing a map, and the "ResultPath" strings from this structure can be referenced by pointer,
+// rather than needing to dynamically allocate a copy of the resolved path.  Lookups by PathInSource can use a much faster
+// case sensitive hash, because PathInSource has verbatim capitalization from the source code files.  Altogether, this
+// utility structure saves a bunch of shader cache query, hash, map, string, and memory allocation overhead.
+//
+struct FShaderPreprocessDependencies
+{
+	// First item in array contains stripped source for root file, and is not in the hash tables
+	TArray<FShaderPreprocessDependency> Dependencies;
+	FHashTable BySource;								// Hash table by PathInSource
+	FHashTable ByResult;								// Hash table by ResultPath
+};
+
+typedef TSharedPtr<FShaderPreprocessDependencies, ESPMode::ThreadSafe> FShaderPreprocessDependenciesShared;
+
+/**
+ * Utility function that returns a root shader file plus all non-platform include dependencies in a single batch call, useful for preprocessing.
+ */
+extern RENDERCORE_API bool GetShaderPreprocessDependencies(const TCHAR* VirtualFilePath, EShaderPlatform ShaderPlatform, FShaderPreprocessDependenciesShared& OutDependencies);
+
 enum class EShaderCompilerWorkerType : uint8
 {
 	None,
