Shader Preprocessor:  Roll TEXT macro substitution into the C preprocessor step, via custom macro callback feature, making it effectively free.  Saves 3% overall on preprocess, but also represents a simplification by reducing the number of unique preprocess passes.

#rnx

diff --git a/Engine/Shaders/Private/ShaderPrintCommon.ush b/Engine/Shaders/Private/ShaderPrintCommon.ush
--- a/Engine/Shaders/Private/ShaderPrintCommon.ush
+++ b/Engine/Shaders/Private/ShaderPrintCommon.ush
@@ -458,10 +458,6 @@ float2 ShaderPrintText_Internal(FShaderPrintContext Ctx, bool bIsActive, float2
 void Print(inout FShaderPrintContext Ctx, FShaderPrintText InTextEntry, FFontColor InColor)	{ Ctx.Pos = ShaderPrintText_Internal(Ctx, Ctx.bIsActive, Ctx.Pos, InTextEntry, InColor); }
 void Print(inout FShaderPrintContext Ctx, FShaderPrintText InTextEntry)						{ Ctx.Pos = ShaderPrintText_Internal(Ctx, Ctx.bIsActive, Ctx.Pos, InTextEntry, GetDefaultFontColor()); }
 
-// Function for reading global TEXT string. 
-// The data and function are generated by the shader compiler.
-GENERATED_SHADER_PRINT
-
 ///////////////////////////////////////////////////////////////////////////////////////////////////
 // Value printing (common value printing)
 #define MAX_DECIMAL_COUNT 5u
diff --git a/Engine/Source/Developer/Apple/MetalShaderFormat/Private/MetalShaderCompiler.cpp b/Engine/Source/Developer/Apple/MetalShaderFormat/Private/MetalShaderCompiler.cpp
--- a/Engine/Source/Developer/Apple/MetalShaderFormat/Private/MetalShaderCompiler.cpp
+++ b/Engine/Source/Developer/Apple/MetalShaderFormat/Private/MetalShaderCompiler.cpp
@@ -874,9 +874,6 @@ bool PreprocessMetalShader(const FShaderCompilerInput& Input, const FShaderCompi
 
 	CleanupUniformBufferCode(Input.Environment, PreprocessOutput.EditSource());
 
-	// Process TEXT macro.
-	TransformStringIntoCharacterArray(PreprocessOutput.EditSource());
-
 	// Run the shader minifier
 	#if UE_METAL_SHADER_COMPILER_ALLOW_DEAD_CODE_REMOVAL
 	if (Input.Environment.CompilerFlags.Contains(CFLAG_RemoveDeadCode))
diff --git a/Engine/Source/Developer/ShaderCompilerCommon/Private/ShaderCompilerCommon.cpp b/Engine/Source/Developer/ShaderCompilerCommon/Private/ShaderCompilerCommon.cpp
--- a/Engine/Source/Developer/ShaderCompilerCommon/Private/ShaderCompilerCommon.cpp
+++ b/Engine/Source/Developer/ShaderCompilerCommon/Private/ShaderCompilerCommon.cpp
@@ -1486,171 +1486,10 @@ static bool SearchText(const TCHAR*& InOut, FAssertParsingStack& OutAssertStack)
 	return InOut != nullptr;
 }
 
-// Simple token matching and expansion to replace TEXT macro into supported character string
+// Deprecated function -- logic now handled in PreprocessShader
 void TransformStringIntoCharacterArray(FString& PreprocessedShaderSource, TArray<FShaderDiagnosticData>* OutDiagnosticDatas)
 {
-	// Early out if input is empty; '&PreprocessedShaderSource[0]' below does not return a valid pointer for empty FString
-	if (PreprocessedShaderSource.IsEmpty())
-	{
-		return;
-	}
-
-	struct FTextEntry
-	{
-		uint32  Index;
-		uint32  Hash;
-		uint32  Offset;
-		bool    bIsAssert;
-		FString SourceText;
-		FString ConvertedText;
-		FString EncodedText;
-	};
-	TArray<FTextEntry> Entries;
-
-	// 1. Find all TEXT strings
-	// 2. Add a text entry
-	// 3. Replace TEXT by its entry number
-	uint32 GlobalCount = 0;
-	uint32 AssertCount = 0;
-	uint32 PrintfCount = 0;
-	{
-		const FString InitHashBegin(TEXT("InitShaderPrintText("));
-		const FString InitHashEnd(TEXT(")"));
-
-		FAssertParsingStack AssertParsingStack;
-		const TCHAR* SearchPtr = &PreprocessedShaderSource[0];
-		while (SearchText(SearchPtr, AssertParsingStack))
-		{
-			const TCHAR* EndPtr = nullptr;
-			FString Text = ParseText(SearchPtr, EndPtr);
-			if (EndPtr)
-			{
-				// Trim enclosing
-				Text.RemoveFromEnd("\")");
-				Text.RemoveFromStart("(\"");
-
-				// Register entry and convert text
-				const uint32 EntryIndex = Entries.Num();
-				uint32 ValidCharCount = 0;
-				FTextEntry& Entry = Entries.AddDefaulted_GetRef();
-				Entry.Index			= EntryIndex;
-				Entry.Offset		= GlobalCount;
-				Entry.SourceText	= Text;
-				ConvertTextToAsciiCharacter(Entry.SourceText, Entry.ConvertedText, Entry.EncodedText);
-				Entry.Hash			= CityHash32((const char*)Entry.SourceText.GetCharArray().GetData(), sizeof(FString::ElementType) * Entry.SourceText.Len());
-				Entry.bIsAssert		= AssertParsingStack.IsAssert();
-				if (Entry.bIsAssert)
-				{
-					++AssertCount;
-				}
-				else
-				{
-					++PrintfCount;
-				}
-
-				// Sanity check
-				uint32 HCheck = CityHash32((const char*)Entry.SourceText.GetCharArray().GetData(), sizeof(FString::ElementType) * Entry.SourceText.Len());
-				check(HCheck == Entry.Hash);
-
-				GlobalCount += Entry.ConvertedText.Len();
-
-				// Replace string
-				const TCHAR* StartPtr = &PreprocessedShaderSource[0];
-				const uint32 StartIndex = SearchPtr - StartPtr;
-				const uint32 CharCount = (EndPtr - SearchPtr) + 1;
-				PreprocessedShaderSource.RemoveAt(StartIndex, CharCount);
-
-				// * If this is an assert, replace the text with a simple hash
-				// * If this is a regular print, replace the text with IniShaderPrintText for runtime lookup
-				if (Entry.bIsAssert)
-				{
-					const FString HashString = FString::Printf(TEXT("%u"), Entry.Hash);
-					PreprocessedShaderSource.InsertAt(StartIndex, HashString);
-				}
-				else
-				{
-					const FString HashText = InitHashBegin + FString::FromInt(EntryIndex) + InitHashEnd;
-					PreprocessedShaderSource.InsertAt(StartIndex, HashText);
-				}
-
-				// Update SearchPtr, as PreprocessedShaderSource has been modified, and its memory could have been reallocated, causing SearchPtr to be invalid.
-				SearchPtr = &PreprocessedShaderSource[0] + StartIndex;
-			}
-		}
-	}
-
-	// 4. Write a global struct containing all the entries
-	// 5. Write the function for fetching character for a given entry index
-	const uint32 EntryCount = Entries.Num();
-	FString TextChars;
-	if (PrintfCount>0 && EntryCount>0 && GlobalCount>0)
-	{
-		// 1. Encoded character for each text entry within a single global char array
-		TextChars = FString::Printf(TEXT("static const uint TEXT_CHARS[%d] = {\n"), GlobalCount);
-		for (FTextEntry& Entry : Entries)
-		{
-			TextChars += FString::Printf(TEXT("\t%s%s // %d: \"%s\"\n"), *Entry.EncodedText, Entry.Index < EntryCount - 1 ? TEXT(",") : TEXT(""), Entry.Index, * Entry.SourceText);
-		}
-		TextChars += TEXT("};\n\n");
-
-		// 2. Offset within the global array
-		TextChars += FString::Printf(TEXT("static const uint TEXT_OFFSETS[%d] = {\n"), EntryCount+1);
-		for (FTextEntry& Entry : Entries)
-		{
-			TextChars += FString::Printf(TEXT("\t%d, // %d: \"%s\"\n"), Entry.Offset, Entry.Index, *Entry.SourceText);
-		}
-		TextChars += FString::Printf(TEXT("\t%d // end\n"), GlobalCount);
-		TextChars += TEXT("};\n\n");
-
-		// 3. Entry hashes
-		TextChars += TEXT("// Hashes are computed using the CityHash32 function\n");
-		TextChars += FString::Printf(TEXT("static const uint TEXT_HASHES[%d] = {\n"), EntryCount);
-		for (FTextEntry& Entry : Entries)
-		{
-			TextChars += FString::Printf(TEXT("\t0x%x%s // %d: \"%s\"\n"), Entry.Hash, Entry.Index < EntryCount - 1 ? TEXT(",") : TEXT(""), Entry.Index, * Entry.SourceText);
-		}
-		TextChars += TEXT("};\n\n");
-
-		TextChars += TEXT("uint ShaderPrintGetChar(uint InIndex)              { return TEXT_CHARS[InIndex]; }\n");
-		TextChars += TEXT("uint ShaderPrintGetOffset(FShaderPrintText InText) { return TEXT_OFFSETS[InText.Index]; }\n");
-		TextChars += TEXT("uint ShaderPrintGetHash(FShaderPrintText InText)   { return TEXT_HASHES[InText.Index]; }\n");
-	}
-	else
-	{	
-		TextChars += TEXT("uint ShaderPrintGetChar(uint Index)                { return 0; }\n");
-		TextChars += TEXT("uint ShaderPrintGetOffset(FShaderPrintText InText) { return 0; }\n");
-		TextChars += TEXT("uint ShaderPrintGetHash(FShaderPrintText InText)   { return 0; }\n");
-	}
-	
-	// 6. Insert global struct data + print function
-	{
-		const TCHAR* InsertToken = TEXT("GENERATED_SHADER_PRINT");
-		const TCHAR* SearchPtr = FCString::Strstr(&PreprocessedShaderSource[0], InsertToken);
-		if (SearchPtr)
-		{
-			// Replace string
-			const TCHAR* StartPtr = &PreprocessedShaderSource[0];
-			const uint32 StartIndex = SearchPtr - StartPtr;
-			const uint32 CharCount = FCString::Strlen(InsertToken);
-			PreprocessedShaderSource.RemoveAt(StartIndex, CharCount);
-			PreprocessedShaderSource.InsertAt(StartIndex, TextChars);
-		}
-	}
-
-	// 7. Insert assert data into shader compilation output for runtime CPU lookup
-	if (OutDiagnosticDatas && AssertCount > 0)
-	{
-		OutDiagnosticDatas->Reserve(OutDiagnosticDatas->Num() + AssertCount);
-		for (const FTextEntry& E : Entries)
-		{
-			if (E.bIsAssert)
-			{
-				FShaderDiagnosticData& Data = OutDiagnosticDatas->AddDefaulted_GetRef();
-				Data.Hash = E.Hash;
-				Data.Message = E.SourceText;
-			}
-		}
-	}
+	// Empty
 }
 
 ///////////////////////////////////////////////////////////////////////////////////////////////////////////////////////////////////////
diff --git a/Engine/Source/Developer/ShaderCompilerCommon/Public/ShaderCompilerCommon.h b/Engine/Source/Developer/ShaderCompilerCommon/Public/ShaderCompilerCommon.h
--- a/Engine/Source/Developer/ShaderCompilerCommon/Public/ShaderCompilerCommon.h
+++ b/Engine/Source/Developer/ShaderCompilerCommon/Public/ShaderCompilerCommon.h
@@ -338,6 +338,7 @@ extern SHADERCOMPILERCOMMON_API const TCHAR* ParseHLSLSymbolName(const TCHAR* Se
 extern SHADERCOMPILERCOMMON_API void ParseHLSLTypeName(const TCHAR* SearchString, const TCHAR*& TypeNameStartPtr, const TCHAR*& TypeNameEndPtr);
 
 // Processes TEXT macros
+UE_DEPRECATED(5.4, "TEXT macro processing is now handled in PreprocessShader --  TransformStringIntoCharacterArray is a no-op")
 extern SHADERCOMPILERCOMMON_API void TransformStringIntoCharacterArray(FString& PreprocessedShaderSource, TArray<FShaderDiagnosticData>* OutDiagnosticDatas=nullptr);
 
 // Structure to hold forward declarations for a specific scope/namespace chain for the HlslParser
diff --git a/Engine/Source/Developer/ShaderFormatOpenGL/Private/OpenGLShaderCompiler.cpp b/Engine/Source/Developer/ShaderFormatOpenGL/Private/OpenGLShaderCompiler.cpp
--- a/Engine/Source/Developer/ShaderFormatOpenGL/Private/OpenGLShaderCompiler.cpp
+++ b/Engine/Source/Developer/ShaderFormatOpenGL/Private/OpenGLShaderCompiler.cpp
@@ -3294,9 +3294,6 @@ bool PreprocessOpenGLShader(const FShaderCompilerInput& Input, const FShaderComp
 	// This requires removing the HLSLCC_NoPreprocess flag later on!
 	CleanupUniformBufferCode(Input.Environment, PreprocessOutput.EditSource());
 
-	// Process TEXT macro.
-	TransformStringIntoCharacterArray(PreprocessOutput.EditSource());
-
 	// Run the experimental shader minifier
 #if UE_OPENGL_SHADER_COMPILER_ALLOW_DEAD_CODE_REMOVAL
 	if (Input.Environment.CompilerFlags.Contains(CFLAG_RemoveDeadCode))
diff --git a/Engine/Source/Developer/ShaderPreprocessor/Private/ShaderPreprocessor.cpp b/Engine/Source/Developer/ShaderPreprocessor/Private/ShaderPreprocessor.cpp
--- a/Engine/Source/Developer/ShaderPreprocessor/Private/ShaderPreprocessor.cpp
+++ b/Engine/Source/Developer/ShaderPreprocessor/Private/ShaderPreprocessor.cpp
@@ -165,16 +165,41 @@ struct FStbPreprocessContext
 	uint32 VertexFactoryOffset = INDEX_NONE;					// Vertex factory dependencies start at this offset in LoadedIncludesCacheShared
 	uint32 VirtualSharedContentsOffset = INDEX_NONE;			// Virtual shared contents start at this offset in LoadedIncludesCacheShared
 
+	// TEXT macro processing state
+	struct FTextEntry
+	{
+		uint32  Index;
+		uint32  Hash;
+		uint32  Offset;
+		bool    bIsAssert;
+		FString SourceText;
+		FString ConvertedText;
+		FString EncodedText;
+	};
+	TArray<FTextEntry> TextEntries;
+	TArray<ANSICHAR> TextMacroSubstituted;
+	uint32 TextGlobalCount = 0;
+	uint32 TextAssertCount = 0;
+	uint32 TextPrintfCount = 0;
+	bool bInAssert = false;
+
+	bool HasIncludedHeader(const FString& Header)
+	{
+		// Checks if a given header  has been included.  Note that the header may be encountered through one of our FShaderPreprocessDependencies structures,
+		// so if those are valid, we need to check the corresponding elements in the LoadedIncludesCacheShared array to see if the path was encountered.
+		return
+			(PreprocessDependencies.IsValid() && HasDependencyFromResultPath(*PreprocessDependencies, Header, &LoadedIncludesCacheShared[0])) ||
+			(VertexFactoryDependencies.IsValid() && HasDependencyFromResultPath(*VertexFactoryDependencies, Header, &LoadedIncludesCacheShared[VertexFactoryOffset])) ||
+			LoadedIncludesCache.Contains(Header);
+	}
+
 	bool HasIncludedMandatoryHeaders()
 	{
-		// Checks if the mandatory PlatformHeader ("/Engine/Public/Platform.ush") has been included.  Note that the header may be
-		// encountered through one of our FShaderPreprocessDependencies structures, so if those are valid, we need to check the
-		// corresponding elements in the LoadedIncludesCacheShared array to see if the path was encountered.
-		return
-			(PreprocessDependencies.IsValid() && HasDependencyFromResultPath(*PreprocessDependencies, PlatformHeader, &LoadedIncludesCacheShared[0])) ||
-			(VertexFactoryDependencies.IsValid() && HasDependencyFromResultPath(*VertexFactoryDependencies, PlatformHeader, &LoadedIncludesCacheShared[VertexFactoryOffset])) ||
-			LoadedIncludesCache.Contains(PlatformHeader);
+		// Check if the mandatory PlatformHeader has been included ("/Engine/Public/Platform.ush")
+		return HasIncludedHeader(PlatformHeader);
 	}
+
+	void ShaderPrintGenerate(char* PreprocessFile, TArray<FShaderDiagnosticData>* OutDiagnosticDatas);
 };
 
 static void StbLoadedIncludeTrimPaddingChecked(FStbLoadedInclude* ContentsCached)
@@ -438,11 +463,203 @@ static const ANSICHAR* StbResolveInclude(const ANSICHAR* PathInSource, uint32 Pa
 	return nullptr;
 }
 
+static const char* ShaderPrintTextIdentifier = "TEXT";
+static const char* ShaderPrintAssertIdentifier = "UEReportAssertWithPayload";
+
+static const char* StbCustomMacroBegin(const char* OriginalText, void* RawContext)
+{
+	FStbPreprocessContext& Context = *reinterpret_cast<FStbPreprocessContext*>(RawContext);
+
+	// Check for assert macro
+	if (FCStringAnsi::Strstr(OriginalText, ShaderPrintAssertIdentifier) == OriginalText)
+	{
+		// We only need to track that we're in an assert, we don't need to do any substitution
+		Context.bInAssert = true;
+		return OriginalText;
+	}
+
+	// TEXT macro
+	check(FCStringAnsi::Strstr(OriginalText, ShaderPrintTextIdentifier) == OriginalText);
+	const char* TextChar = OriginalText;
+	while (*TextChar != '(')
+	{
+		TextChar++;
+	}
+	TextChar++;
+	while (*TextChar != ')' && *TextChar != '\"')
+	{
+		TextChar++;
+	}
+
+	// If no quoted text, that's a parse error
+	if (*TextChar != '\"')
+	{
+		return nullptr;
+	}
+
+	// We found a string, add an entry
+	const uint32 EntryIndex = Context.TextEntries.Num();
+	FStbPreprocessContext::FTextEntry& Entry = Context.TextEntries.AddDefaulted_GetRef();
+	Entry.Index = EntryIndex;
+	Entry.Offset = Context.TextGlobalCount;
+	Entry.bIsAssert = Context.bInAssert;
+	if (Entry.bIsAssert)
+	{
+		++Context.TextAssertCount;
+	}
+	else
+	{
+		++Context.TextPrintfCount;
+	}
+	
+	// Parse the string, handling escaped characters.  SourceText contains the raw text, ConvertedText removes escape back slashes,
+	// and EncodedText is an array of integer numeric values as ASCII.
+	TextChar++;
+	const char* TextStart = TextChar;
+	int32 CharCount = 0;
+	for (; *TextChar != '\"'; TextChar++)
+	{
+		if (*TextChar == '\\')
+		{
+			TextChar++;
+		}
+		CharCount++;
+	}
+	Entry.SourceText = FString(FAnsiStringView(TextStart, TextChar - TextStart));
+	Entry.ConvertedText.GetCharArray().SetNumUninitialized(CharCount + 1);
+	Entry.EncodedText.Reserve(CharCount * 4); // ~3 digits per character + a comma
+
+	TCHAR* ConvertedTextData = Entry.ConvertedText.GetCharArray().GetData();
+
+	int32 CharIndex = 0;
+	for (TextChar = TextStart; *TextChar != '\"'; TextChar++, CharIndex++)
+	{
+		if (*TextChar == '\\')
+		{
+			TextChar++;
+		}
+		ConvertedTextData[CharIndex] = *TextChar;
+
+		const char C = *TextChar;
+		Entry.EncodedText.AppendInt(uint8(C));
+		if (CharIndex + 1 != CharCount)
+		{
+			Entry.EncodedText += ',';
+		}
+	}
+	check(CharIndex == CharCount);
+	ConvertedTextData[CharIndex] = 0;
+
+	Entry.Hash = CityHash32((const char*)Entry.SourceText.GetCharArray().GetData(), sizeof(FString::ElementType) * Entry.SourceText.Len());
+	Context.TextGlobalCount += Entry.ConvertedText.Len();
+
+	// Generate substitution string
+	if (Entry.bIsAssert)
+	{
+		const FString HashString = FString::Printf(TEXT("%u"), Entry.Hash);
+		CopyStringToAnsiCharArray(*HashString, HashString.Len(), Context.TextMacroSubstituted);
+	}
+	else
+	{
+		const FString InitHashBegin(TEXT("InitShaderPrintText("));
+		const FString InitHashEnd(TEXT(")"));
+
+		const FString HashText = InitHashBegin + FString::FromInt(EntryIndex) + InitHashEnd;
+		CopyStringToAnsiCharArray(*HashText, HashText.Len(), Context.TextMacroSubstituted);
+	}
+
+	return Context.TextMacroSubstituted.GetData();
+}
+
+static void StbCustomMacroEnd(const char* OriginalText, void* RawContext, const char* SubstitutionText)
+{
+	FStbPreprocessContext& Context = *reinterpret_cast<FStbPreprocessContext*>(RawContext);
+
+	if (FCStringAnsi::Strstr(OriginalText, ShaderPrintAssertIdentifier) == OriginalText)
+	{
+		Context.bInAssert = false;
+	}
+}
+
+void FStbPreprocessContext::ShaderPrintGenerate(char* PreprocessedFile, TArray<FShaderDiagnosticData>* OutDiagnosticDatas)
+{
+	// Check if ShaderPrintCommon.ush was included, to decide whether to add the shader print generated code
+	static FString ShaderPrintHeader("/Engine/Private/ShaderPrintCommon.ush");
+	if (!HasIncludedHeader(ShaderPrintHeader))
+	{
+		return;
+	}
+
+	// 1. Write a global struct containing all the entries
+	// 2. Write the function for fetching character for a given entry index
+	const uint32 EntryCount = TextEntries.Num();
+	FString TextChars;
+	if (TextPrintfCount > 0 && EntryCount > 0 && TextGlobalCount > 0)
+	{
+		// 1. Encoded character for each text entry within a single global char array
+		TextChars = FString::Printf(TEXT("\n\nstatic const uint TEXT_CHARS[%d] = {\n"), TextGlobalCount);
+		for (FTextEntry& Entry : TextEntries)
+		{
+			TextChars += FString::Printf(TEXT("\t%s%s // %d: \"%s\"\n"), *Entry.EncodedText, Entry.Index < EntryCount - 1 ? TEXT(",") : TEXT(""), Entry.Index, *Entry.SourceText);
+		}
+		TextChars += TEXT("};\n\n");
+
+		// 2. Offset within the global array
+		TextChars += FString::Printf(TEXT("static const uint TEXT_OFFSETS[%d] = {\n"), EntryCount + 1);
+		for (FTextEntry& Entry : TextEntries)
+		{
+			TextChars += FString::Printf(TEXT("\t%d, // %d: \"%s\"\n"), Entry.Offset, Entry.Index, *Entry.SourceText);
+		}
+		TextChars += FString::Printf(TEXT("\t%d // end\n"), TextGlobalCount);
+		TextChars += TEXT("};\n\n");
+
+		// 3. Entry hashes
+		TextChars += TEXT("// Hashes are computed using the CityHash32 function\n");
+		TextChars += FString::Printf(TEXT("static const uint TEXT_HASHES[%d] = {\n"), EntryCount);
+		for (FTextEntry& Entry : TextEntries)
+		{
+			TextChars += FString::Printf(TEXT("\t0x%x%s // %d: \"%s\"\n"), Entry.Hash, Entry.Index < EntryCount - 1 ? TEXT(",") : TEXT(""), Entry.Index, *Entry.SourceText);
+		}
+		TextChars += TEXT("};\n\n");
+
+		TextChars += TEXT("uint ShaderPrintGetChar(uint InIndex)              { return TEXT_CHARS[InIndex]; }\n");
+		TextChars += TEXT("uint ShaderPrintGetOffset(FShaderPrintText InText) { return TEXT_OFFSETS[InText.Index]; }\n");
+		TextChars += TEXT("uint ShaderPrintGetHash(FShaderPrintText InText)   { return TEXT_HASHES[InText.Index]; }\n");
+	}
+	else
+	{
+		TextChars += TEXT("uint ShaderPrintGetChar(uint Index)                { return 0; }\n");
+		TextChars += TEXT("uint ShaderPrintGetOffset(FShaderPrintText InText) { return 0; }\n");
+		TextChars += TEXT("uint ShaderPrintGetHash(FShaderPrintText InText)   { return 0; }\n");
+	}
+
+	// 3. Insert global struct data + print function
+	TArray<ANSICHAR> TextCharsAnsi;
+	CopyStringToAnsiCharArray(*TextChars, TextChars.Len(), TextCharsAnsi);
+
+	preprocessor_file_append(PreprocessedFile, TextCharsAnsi.GetData(), TextCharsAnsi.Num() - 1);
+
+	// 4. Insert assert data into shader compilation output for runtime CPU lookup
+	if (OutDiagnosticDatas && TextAssertCount > 0)
+	{
+		OutDiagnosticDatas->Reserve(OutDiagnosticDatas->Num() + TextAssertCount);
+		for (const FTextEntry& E : TextEntries)
+		{
+			if (E.bIsAssert)
+			{
+				FShaderDiagnosticData& Data = OutDiagnosticDatas->AddDefaulted_GetRef();
+				Data.Hash = E.Hash;
+				Data.Message = E.SourceText;
+			}
+		}
+	}
+}
+
 class FShaderPreprocessorModule : public IModuleInterface
 {
 	virtual void StartupModule() override
 	{
-		init_preprocessor(&StbLoadFile, &StbFreeFile, &StbResolveInclude);
+		init_preprocessor(&StbLoadFile, &StbFreeFile, &StbResolveInclude, &StbCustomMacroBegin, &StbCustomMacroEnd);
 		// disable the "directive not at start of line" error; this allows a few things:
 		// 1. #define'ing #pragma messages - consumed by the preprocessor (to handle UESHADERMETADATA hackery)
 		// 2. #define'ing other #pragmas (those not processed explicitly by the preprocessor are copied into the preprocessed code
@@ -485,6 +702,9 @@ bool InnerPreprocessShaderStb(
 	macro_definition** StbDefines = nullptr;
 	FShaderPreprocessorUtilities::PopulateDefines(Environment, AdditionalDefines, &MacroArena, StbDefines);
 
+	arrput(StbDefines, pp_define_custom_macro(&MacroArena, ShaderPrintTextIdentifier));
+	arrput(StbDefines, pp_define_custom_macro(&MacroArena, ShaderPrintAssertIdentifier));
+
 	PRAGMA_ENABLE_DEPRECATION_WARNINGS
 
 	FStbPreprocessContext Context{ Input, Environment };
@@ -594,6 +814,9 @@ bool InnerPreprocessShaderStb(
 
 	if (!HasError)
 	{
+		// Append ShaderPrint generated code at the end of the shader if necessary
+		Context.ShaderPrintGenerate(OutPreprocessedAnsi, &Output.EditDiagnosticDatas());
+
 		// "preprocessor_file_size" includes null terminator, so subtract one for Append call -- passing size saves an expensive strlen in Append
 		Output.EditSource().Append(OutPreprocessedAnsi, preprocessor_file_size(OutPreprocessedAnsi) - 1);
 	}
diff --git a/Engine/Source/Developer/ShaderPreprocessor/Private/stb_preprocess/preprocessor.c b/Engine/Source/Developer/ShaderPreprocessor/Private/stb_preprocess/preprocessor.c
--- a/Engine/Source/Developer/ShaderPreprocessor/Private/stb_preprocess/preprocessor.c
+++ b/Engine/Source/Developer/ShaderPreprocessor/Private/stb_preprocess/preprocessor.c
@@ -104,6 +104,7 @@ struct macro_definition
 #define MACRO_NUM_PARAMETERS_line -3
 #define MACRO_NUM_PARAMETERS_defined -4
 #define MACRO_NUM_PARAMETERS_counter -5
+#define MACRO_NUM_PARAMETERS_custom -6
 
 static struct macro_definition predefined_FILE = {"__FILE__", 8, {0}, 0, MACRO_NUM_PARAMETERS_file, 1};
 static struct macro_definition predefined_LINE = {"__LINE__", 8, {0}, 0, MACRO_NUM_PARAMETERS_line, 1};
@@ -172,6 +173,8 @@ typedef struct
 static loadfile_callback_func loadfile_callback;
 static freefile_callback_func freefile_callback;
 static resolveinclude_callback_func resolveinclude_callback;
+static custommacro_begin_callback_func custommacro_begin;
+static custommacro_end_callback_func custommacro_end;
 
 #define HASH_EMPTY_MARKER 0
 #define HASH_TOMBSTONE 1
@@ -3126,6 +3129,76 @@ static void maybe_expand_macro(parse_state* cs, struct macro_definition* pending
 				++c->counter;
 				break;
 			}
+			case MACRO_NUM_PARAMETERS_custom:
+			{
+				while (char_is_whitespace(*p))
+					++p;
+
+				// If no parentheses or the macro is disabled (say because we are encountering it nested), just leave the echoed identifier
+				if (*p != '(' || md->disabled)
+				{
+					arrsetlen(cs->dest, (size_t)(q - cs->dest));  // set the output buffer size to account for the above-copied identifier
+					return;
+				}
+				++p;
+
+				// Generate a buffer and append the rest of the arguments to it
+				int arg_newlines = 0;
+				char* custom_macro_buffer = 0;
+				arrsetcap(custom_macro_buffer, identifier_length + 512);
+				arrsetlen(custom_macro_buffer, identifier_length + 1);
+				memcpy(custom_macro_buffer, identifier, identifier_length);
+				custom_macro_buffer[identifier_length] = '(';
+
+				for (;;)
+				{
+					p = copy_argument(p, &arg_newlines, &custom_macro_buffer);
+					if (*p == 0)
+					{
+						arrfree(custom_macro_buffer);
+						do_error(cs, "End-of-file in macro '%s' argument list", md->symbol_name);
+						return;
+					}
+
+					// Add the comma or close parentheses to macro buffer
+					arrput(custom_macro_buffer, *p);
+					if (*p++ == ')')
+					{
+						break;
+					}
+				}
+
+				// Null terminate
+				arrput(custom_macro_buffer, 0);
+
+				// Send the custom macro text to the callback, preprocess it, then signal the end of the custom macro
+				const char* substitution_text = custommacro_begin(custom_macro_buffer, c->custom_context);
+
+				if (substitution_text == 0)
+				{
+					arrfree(custom_macro_buffer);
+					do_error(cs, "Custom macro handler for '%s' failed", md->symbol_name);
+					return;
+				}
+
+				parse_state ncs = *cs;
+				ncs.parent = cs;
+				ncs.src = substitution_text;
+				ncs.src_offset = 0;
+				ncs.src_length = strlen(substitution_text);
+				md->disabled = 1;
+				preprocess_string(&ncs, IN_MACRO_yes, md->symbol_name, NULL);
+				md->disabled = 0;
+				cs->dest = ncs.dest;
+
+				custommacro_end(custom_macro_buffer, c->custom_context, substitution_text);
+				arrfree(custom_macro_buffer);
+
+				// Set source offset to the end of the macro
+				cs->src_offset = p - cs->src;
+				cs->src_line_number += arg_newlines;
+				break;
+			}
 		}
 	}
 	else
@@ -4764,6 +4837,21 @@ static void define_macro(pphash* map, struct macro_definition* m)
 	stringhash_put(map, m->symbol_name, m->symbol_name_length, m);
 }
 
+struct macro_definition* pp_define_custom_macro(struct stb_arena* a, const char* identifier)
+{
+	struct macro_definition* md = (struct macro_definition*)stb_arena_alloc(a, sizeof(*md));
+
+	memset(md, 0, sizeof(*md));
+	md->symbol_name = (char*)identifier;
+	md->symbol_name_length = strlen(identifier);
+	md->simple_expansion = 0;
+	md->simple_expansion_length = 0;
+	md->num_parameters = MACRO_NUM_PARAMETERS_custom;
+	md->predefined = 1;
+
+	return md;
+}
+
 char* preprocess_file(const char* filename,
 	void* custom_context,
 	struct macro_definition** predefined_macros,
@@ -4878,6 +4966,23 @@ int preprocessor_file_capacity(char* text)
 	return text ? arrcap(text) : 0;
 }
 
+void preprocessor_file_append(char* text, const char* appended_text, int appended_text_len)
+{
+	if (text)
+	{
+		int text_len = arrlen(text);
+
+		// The preprocessor text array length includes the null terminator, so we don't need to add one here
+		arrsetlen(text, text_len + appended_text_len);
+
+		// Subtract one, so we start writing at original null terminator
+		memcpy(text + text_len - 1, appended_text, appended_text_len);
+
+		// And add a new null terminator
+		text[text_len + appended_text_len - 1] = 0;
+	}
+}
+
 void preprocessor_file_free(char* text, pp_diagnostic* pd)
 {
 	int i, j;
@@ -4961,7 +5066,9 @@ static void init_char_type(void)
 void init_preprocessor(
 	loadfile_callback_func load_callback, 
 	freefile_callback_func free_callback,
-	resolveinclude_callback_func resolve_callback)
+	resolveinclude_callback_func resolve_callback,
+	custommacro_begin_callback_func custommacro_begin_callback,
+	custommacro_end_callback_func custommacro_end_callback)
 {
 #ifdef HASHTEST
 	test_stringhash();
@@ -4987,6 +5094,8 @@ void init_preprocessor(
 	loadfile_callback = load_callback;
 	freefile_callback = free_callback;
 	resolveinclude_callback = resolve_callback;
+	custommacro_begin = custommacro_begin_callback;
+	custommacro_end = custommacro_end_callback;
 }
 
 #pragma warning(pop)
diff --git a/Engine/Source/Developer/ShaderPreprocessor/Private/stb_preprocess/preprocessor.h b/Engine/Source/Developer/ShaderPreprocessor/Private/stb_preprocess/preprocessor.h
--- a/Engine/Source/Developer/ShaderPreprocessor/Private/stb_preprocess/preprocessor.h
+++ b/Engine/Source/Developer/ShaderPreprocessor/Private/stb_preprocess/preprocessor.h
@@ -75,11 +75,25 @@ typedef void (*freefile_callback_func)(const char* filename, const char* loaded_
 // note: returned resolved paths must remain valid for the lifetime of a single preprocessor execution
 typedef const char* (*resolveinclude_callback_func)(const char* path, unsigned int path_len, const char* parent, void* custom_context);
 
+// Callback functions for custom macros.  Custom macros return their own substitution text, which is then further
+// preprocessed.  The text passed to the callback includes the macro identifier and its arguments, with backslashes
+// removed.  A function is called when entering the macro, and again when complete, allowing stateful logic
+// based on whether parsing is inside a certain custom macro.  Any custom macro will automatically be disabled
+// inside itself, with the idea that this allows a singleton buffer for the substitution text, as there's no
+// concern about nested calls.  If text is dynamically allocated, the caller is responsible for freeing it.  It's
+// also valid to return the original text passed in as the substitution text, if you just want to set some state
+// inside the context, and not actually change the macro text -- since the macro is disabled, it will then be
+// echoed as non-macro text.
+typedef const char* (*custommacro_begin_callback_func)(const char* original_text, void* custom_context);
+typedef void (*custommacro_end_callback_func)(const char* original_text, void* custom_context, const char* substitution_text);
+
 // init function must be called once before any invocations of preprocess_file
 STB_PP_DEF void init_preprocessor(
 	loadfile_callback_func load_callback,
 	freefile_callback_func free_callback,
-	resolveinclude_callback_func resolve_callback);
+	resolveinclude_callback_func resolve_callback,
+	custommacro_begin_callback_func custommacro_begin_callback,
+	custommacro_end_callback_func custommacro_end_callback);
 
 // main preprocessing execution function; returns the preprocessed string and sets any diagnostic messages
 // in the pd array (setting num_pd to the diagnostic count). preprocessor_file_free should be called for 
@@ -97,6 +111,10 @@ STB_PP_DEF char* preprocess_file(
 STB_PP_DEF int preprocessor_file_size(char* text);
 STB_PP_DEF int preprocessor_file_capacity(char* text);
 
+// Append text to the end of a file generated with preprocess_file.  Can save reallocation overhead relative
+// to appending the text later, as there will usually be slack space available.
+STB_PP_DEF void preprocessor_file_append(char* text, const char* appended_text, int appended_text_len);
+
 // frees memory allocated by preprocess_file (preprocessed results and diagnostic messages)
 STB_PP_DEF void preprocessor_file_free(char* text, pp_diagnostic* pd);
 
@@ -104,6 +122,9 @@ STB_PP_DEF void preprocessor_file_free(char* text, pp_diagnostic* pd);
 // DEFINE_NAME DEFINE_VALUE
 STB_PP_DEF struct macro_definition* pp_define(struct stb_arena* a, const char* def);
 
+// specifies a custom macro name, with substitution to be handled by the custom macro callback functions
+STB_PP_DEF struct macro_definition* pp_define_custom_macro(struct stb_arena* a, const char* identifier);
+
 // override behaviour of a particular error result (can be used to, for instance, downgrade
 // an error to a warning or a "no warning" i.e. valid case).
 STB_PP_DEF void pp_set_warning_mode(int result_code, int result_mode);
diff --git a/Engine/Source/Developer/VulkanShaderFormat/Private/VulkanShaderCompiler.cpp b/Engine/Source/Developer/VulkanShaderFormat/Private/VulkanShaderCompiler.cpp
--- a/Engine/Source/Developer/VulkanShaderFormat/Private/VulkanShaderCompiler.cpp
+++ b/Engine/Source/Developer/VulkanShaderFormat/Private/VulkanShaderCompiler.cpp
@@ -2235,9 +2235,6 @@ bool PreprocessVulkanShader(const FShaderCompilerInput& Input, const FShaderComp
 
 	CleanupUniformBufferCode(Input.Environment, PreprocessedShaderSource);
 
-	// Process TEXT macro.
-	TransformStringIntoCharacterArray(PreprocessedShaderSource);
-
 	// Run the shader minifier
 	if (InternalState.IsRayTracingShader() || InternalState.bUseBindlessUniformBuffer)
 	{
diff --git a/Engine/Source/Developer/Windows/ShaderFormatD3D/Private/D3DShaderCompiler.cpp b/Engine/Source/Developer/Windows/ShaderFormatD3D/Private/D3DShaderCompiler.cpp
--- a/Engine/Source/Developer/Windows/ShaderFormatD3D/Private/D3DShaderCompiler.cpp
+++ b/Engine/Source/Developer/Windows/ShaderFormatD3D/Private/D3DShaderCompiler.cpp
@@ -1317,9 +1317,6 @@ bool PreprocessD3DShader(
 
 	CleanupUniformBufferCode(Input.Environment, PreprocessedSource);
 
-	// Process TEXT macro.
-	TransformStringIntoCharacterArray(PreprocessedSource, &Output.EditDiagnosticDatas());
-
 	// Run the shader minifier
 	#if UE_D3D_SHADER_COMPILER_ALLOW_DEAD_CODE_REMOVAL
 	if (Environment.CompilerFlags.Contains(CFLAG_RemoveDeadCode))
