Shader Preprocessor:  ASAN bug fix for stack use after return for undef_map identifiers, by enabling arena allocation.  Strip out whitespace on blank lines -- doesn't measurably affect perf, and could improve deduplication or avoid certain shader invalidations, and makes stripped source more readable.  Fix bug where line numbers aren't incremented when exiting the scanning loop on an identifier.

#rnx

diff --git a/Engine/Source/Developer/ShaderPreprocessor/Private/stb_preprocess/preprocessor.c b/Engine/Source/Developer/ShaderPreprocessor/Private/stb_preprocess/preprocessor.c
--- a/Engine/Source/Developer/ShaderPreprocessor/Private/stb_preprocess/preprocessor.c
+++ b/Engine/Source/Developer/ShaderPreprocessor/Private/stb_preprocess/preprocessor.c
@@ -36,6 +36,9 @@
 #define SSE_READ_PADDING 0
 #endif
 
+// Strips whitespace on blank lines and leading whitespace in directives
+#define STRIP_BLANK_LINE_WHITESPACE 1
+
 #pragma warning(push)
 // arrput calls are incorrectly reported by MSVC analysis to be potentially dereferencing null
 // (the maybegrow macro will inline allocate if given a null pointer to begin with)
@@ -1454,6 +1457,10 @@ static int copy_to_action_point(parse_state* cs)
 		arrsetcap(out, out_len + (cs->src_length - cs->src_offset) + SSE_READ_PADDING);
 		q = out + out_len;
 
+#if STRIP_BLANK_LINE_WHITESPACE
+		char* q_linestart = q;
+#endif
+
 		{  // make lifetime of line_count unambiguous
 			const int state_limit = cs->state_limit;
 			int line_count = 0;
@@ -1461,6 +1468,27 @@ static int copy_to_action_point(parse_state* cs)
 		resume_parsing:
 			do
 			{
+#if STRIP_BLANK_LINE_WHITESPACE
+				uint8 ch = (uint8)*p++;
+				uint8 ch_class;
+				assert(q < out + arrcap(out));
+				STB_ASSUME(q != NULL);
+				ch_class = pp_char_class[ch];
+				prev_state = state;
+				STB_ASSUME(state < PP_STATE_active_count && ch_class < PP_CHAR_CLASS_count);
+				state = pp_transition_table[state][ch_class];
+				if (ch == '\n')				// pp_state_is_end_of_line[state] -- Assume newlines are normalized
+				{
+					// Eliminate leading whitespace before a newline, by backing up to the line start
+					if (prev_state == PP_STATE_in_leading_whitespace)
+					{
+						q = q_linestart;
+					}
+					line_count++;
+					q_linestart = q + 1;
+				}
+				*q++ = ch;
+#else
 				uint8 ch = (uint8)*p++;
 				uint8 ch_class;
 				assert(q < out + arrcap(out));
@@ -1471,6 +1499,7 @@ static int copy_to_action_point(parse_state* cs)
 				STB_ASSUME(state < PP_STATE_active_count && ch_class < PP_CHAR_CLASS_count);
 				state = pp_transition_table[state][ch_class];
 				line_count += (ch == '\n') ? 1 : 0;				// pp_state_is_end_of_line[state] -- Assume newlines are normalized, avoid memory read
+#endif
 			} while (state < state_limit);
 
 			if (state == PP_STATE_saw_identifier_start)
@@ -1491,6 +1520,9 @@ static int copy_to_action_point(parse_state* cs)
 				}
 				else
 				{
+					// Breaking out of the loop skips the line increment below, so do it here as well
+					cs->src_line_number += line_count;
+					cs->dest_line_number += line_count;
 					break;
 				}
 			}
@@ -1515,7 +1547,18 @@ static int copy_to_action_point(parse_state* cs)
 		// identifier case
 
 		if (state > PP_STATE_saw_backslash)
+		{
+#if STRIP_BLANK_LINE_WHITESPACE
+			if (state == PP_STATE_saw_leading_hash)
+			{
+				// Eliminate leading whitespace before a hash, by backing up to the line start.  Need to add one
+				// as we subtract one below to remove what would have been the last character copied (the hash),
+				// if we weren't backing up to line start.
+				q = q_linestart + 1;
+			}
+#endif
 			break;
+		}
 
 		if (state == PP_STATE_saw_backslash)
 		{
@@ -4893,6 +4936,9 @@ char* preprocess_file(const char* filename,
 	// and (b) so most things hit in there first slot. but going too much bigger can
 	// backfire by slowing down the initialization
 	stringhash_create(&c.macro_map, 4096);
+	
+	// enable string arena allocation for undef_map, as the map is passed stack memory identifiers that need to be duplicated
+	stbds_sh_new_arena(c.undef_map);
 
 	memset(c.macro_bloom_filter, 0, sizeof(c.macro_bloom_filter));
 
