Shader Preprocessor:  Fix preprocessing for shader asserts.  Need to preprocess the contents of TEXT macros to handle __FILE__ macros, or other possible macro generated text.

#rnx

diff --git a/Engine/Source/Developer/ShaderPreprocessor/Private/ShaderPreprocessor.cpp b/Engine/Source/Developer/ShaderPreprocessor/Private/ShaderPreprocessor.cpp
--- a/Engine/Source/Developer/ShaderPreprocessor/Private/ShaderPreprocessor.cpp
+++ b/Engine/Source/Developer/ShaderPreprocessor/Private/ShaderPreprocessor.cpp
@@ -702,8 +702,14 @@ bool InnerPreprocessShaderStb(
 	macro_definition** StbDefines = nullptr;
 	FShaderPreprocessorUtilities::PopulateDefines(Environment, AdditionalDefines, &MacroArena, StbDefines);
 
-	arrput(StbDefines, pp_define_custom_macro(&MacroArena, ShaderPrintTextIdentifier));
-	arrput(StbDefines, pp_define_custom_macro(&MacroArena, ShaderPrintAssertIdentifier));
+	// The substitution text generated by custom macros gets run through the preprocessor afterwards, but in some cases we want to
+	// run the arguments through the preprocessor before as well.  The TEXT macro needs this to handle things like TEXT(__FILE__),
+	// where the __FILE__ macro needs to be expanded before the custom macro handler is called, so we pass "1" to enable running
+	// the preprocessor first.  By contrast, for shader asserts, we must NOT run the preprocessor on the arguments first, because
+	// the assert macro sets a state flag which modifies behavior of TEXT macros inside the assert.  Asserts store their TEXT tokens
+	// outside the shader for printing in code when an assert is triggered, while ShaderPrint stores TEXT in the shader itself.
+	arrput(StbDefines, pp_define_custom_macro(&MacroArena, ShaderPrintTextIdentifier, 1));
+	arrput(StbDefines, pp_define_custom_macro(&MacroArena, ShaderPrintAssertIdentifier, 0));
 
 	PRAGMA_ENABLE_DEPRECATION_WARNINGS
 
diff --git a/Engine/Source/Developer/ShaderPreprocessor/Private/stb_preprocess/preprocessor.c b/Engine/Source/Developer/ShaderPreprocessor/Private/stb_preprocess/preprocessor.c
--- a/Engine/Source/Developer/ShaderPreprocessor/Private/stb_preprocess/preprocessor.c
+++ b/Engine/Source/Developer/ShaderPreprocessor/Private/stb_preprocess/preprocessor.c
@@ -99,6 +99,7 @@ struct macro_definition
 	uint8 predefined;
 	uint8 is_variadic;	// if true, then __VA_ARGS__ will be specified using 'num_arguments' as argument_index
 	uint8 disabled;
+	uint8 preprocess_args_first;	// only applies to custom macros
 	struct macro_definition* next;	// push/pop stack
 };
 
@@ -3214,11 +3215,38 @@ static void maybe_expand_macro(parse_state* cs, struct macro_definition* pending
 				// Null terminate
 				arrput(custom_macro_buffer, 0);
 
+				// Optionally preprocess the macro args before calling the custom macro callback.  Need to disable the macro first.
+				md->disabled = 1;
+
+				if (md->preprocess_args_first)
+				{
+					parse_state ncs = *cs;
+					ncs.parent = cs;
+					ncs.src = custom_macro_buffer;
+					ncs.src_offset = 0;
+					ncs.src_length = arrlennonull(custom_macro_buffer) - 1;
+					ncs.dest = 0											;		// Create new dest string
+					preprocess_string(&ncs, IN_MACRO_yes, md->symbol_name, NULL);
+
+					// Replace custom_macro_buffer with new dest string
+					arrfree(custom_macro_buffer);
+					custom_macro_buffer = ncs.dest;
+					
+					if (!custom_macro_buffer)
+					{
+						md->disabled = 0;
+						do_error(cs, "Preprocessing args for custom macro handler '%s' failed", md->symbol_name);
+						return;
+					}
+					arrput(custom_macro_buffer, 0);
+				}
+
 				// Send the custom macro text to the callback, preprocess it, then signal the end of the custom macro
 				const char* substitution_text = custommacro_begin(custom_macro_buffer, c->custom_context);
 
 				if (substitution_text == 0)
 				{
+					md->disabled = 0;
 					arrfree(custom_macro_buffer);
 					do_error(cs, "Custom macro handler for '%s' failed", md->symbol_name);
 					return;
@@ -3229,7 +3257,6 @@ static void maybe_expand_macro(parse_state* cs, struct macro_definition* pending
 				ncs.src = substitution_text;
 				ncs.src_offset = 0;
 				ncs.src_length = strlen(substitution_text);
-				md->disabled = 1;
 				preprocess_string(&ncs, IN_MACRO_yes, md->symbol_name, NULL);
 				md->disabled = 0;
 				cs->dest = ncs.dest;
@@ -4880,7 +4907,7 @@ static void define_macro(pphash* map, struct macro_definition* m)
 	stringhash_put(map, m->symbol_name, m->symbol_name_length, m);
 }
 
-struct macro_definition* pp_define_custom_macro(struct stb_arena* a, const char* identifier)
+struct macro_definition* pp_define_custom_macro(struct stb_arena* a, const char* identifier, unsigned char preprocess_args_first)
 {
 	struct macro_definition* md = (struct macro_definition*)stb_arena_alloc(a, sizeof(*md));
 
@@ -4891,6 +4918,7 @@ struct macro_definition* pp_define_custom_macro(struct stb_arena* a, const char*
 	md->simple_expansion_length = 0;
 	md->num_parameters = MACRO_NUM_PARAMETERS_custom;
 	md->predefined = 1;
+	md->preprocess_args_first = preprocess_args_first;
 
 	return md;
 }
diff --git a/Engine/Source/Developer/ShaderPreprocessor/Private/stb_preprocess/preprocessor.h b/Engine/Source/Developer/ShaderPreprocessor/Private/stb_preprocess/preprocessor.h
--- a/Engine/Source/Developer/ShaderPreprocessor/Private/stb_preprocess/preprocessor.h
+++ b/Engine/Source/Developer/ShaderPreprocessor/Private/stb_preprocess/preprocessor.h
@@ -122,8 +122,9 @@ STB_PP_DEF void preprocessor_file_free(char* text, pp_diagnostic* pd);
 // DEFINE_NAME DEFINE_VALUE
 STB_PP_DEF struct macro_definition* pp_define(struct stb_arena* a, const char* def);
 
-// specifies a custom macro name, with substitution to be handled by the custom macro callback functions
-STB_PP_DEF struct macro_definition* pp_define_custom_macro(struct stb_arena* a, const char* identifier);
+// specifies a custom macro name, with substitution to be handled by the custom macro callback functions.  If "preprocess_args_first" is set,
+// args of macro are preprocessed before begin callback (and again afterwards).
+STB_PP_DEF struct macro_definition* pp_define_custom_macro(struct stb_arena* a, const char* identifier, unsigned char preprocess_args_first);
 
 // override behaviour of a particular error result (can be used to, for instance, downgrade
 // an error to a warning or a "no warning" i.e. valid case).
