- Worked around crash in shader preprocessor on mac

diff --git a/Engine/Source/Developer/ShaderPreprocessor/Private/stb_preprocess/preprocessor.c b/Engine/Source/Developer/ShaderPreprocessor/Private/stb_preprocess/preprocessor.c
--- a/Engine/Source/Developer/ShaderPreprocessor/Private/stb_preprocess/preprocessor.c
+++ b/Engine/Source/Developer/ShaderPreprocessor/Private/stb_preprocess/preprocessor.c
@@ -3225,7 +3225,10 @@ static void maybe_expand_macro(parse_state* cs, struct macro_definition* pending
 					ncs.src = custom_macro_buffer;
 					ncs.src_offset = 0;
 					ncs.src_length = arrlennonull(custom_macro_buffer) - 1;
-					ncs.dest = 0											;		// Create new dest string
+					ncs.dest = 0;
+					// Add some padding, preprocess_string is crashing on Mac, and until we figure root cause, this workaround is removing the known crashes
+					arrsetcap(ncs.dest, ncs.src_length + 512);
+					// Create new dest string
 					preprocess_string(&ncs, IN_MACRO_yes, md->symbol_name, NULL);
 
 					// Replace custom_macro_buffer with new dest string
