[shaders] debug usf/direct compile cleanup
- never append the environment defines as commented code to the source used for further preprocessing/compilation; instead only append it to the debug USF
- strip comments after loading the debug usf in direct compile mode as some backends expect comments to have already been removed and the extra ones we add to the debug dump cause them to barf
- change all #if 0s in the debug usf to block comments instead so the above can strip them (said backends also don't like preprocessor directives left in the file)

diff --git a/Engine/Source/Developer/ShaderCompilerCommon/Private/ShaderCompilerCommon.cpp b/Engine/Source/Developer/ShaderCompilerCommon/Private/ShaderCompilerCommon.cpp
--- a/Engine/Source/Developer/ShaderCompilerCommon/Private/ShaderCompilerCommon.cpp
+++ b/Engine/Source/Developer/ShaderCompilerCommon/Private/ShaderCompilerCommon.cpp
@@ -1905,9 +1905,9 @@ namespace UE::ShaderCompilerCommon
 		}
 	}
 
-	static const TCHAR* Base64EnvBegin = TEXT("#if 0 /* BASE64_ENV */\n");
+	static const TCHAR* Base64EnvBegin = TEXT("/* BASE64_ENV\n");
 	static const int32 Base64EnvBeginLen = FCString::Strlen(Base64EnvBegin);
-	static const TCHAR* Base64EnvEnd = TEXT("\n#endif /* BASE64_ENV */\n");
+	static const TCHAR* Base64EnvEnd = TEXT("\nBASE64_ENV */\n");
 	
 	FString SerializeEnvironmentToBase64(const FShaderCompilerEnvironment& Env)
 	{
@@ -1940,12 +1940,6 @@ namespace UE::ShaderCompilerCommon
 
 	FString GetDebugShaderContents(const FShaderCompilerInput& Input, const FString& PreprocessedSource, const FDebugShaderDataOptions& Options)
 	{
-		FString Contents = Options.AppendPreSource ? Options.AppendPreSource() : FString();
-		Contents += PreprocessedSource;
-		if (Options.AppendPostSource)
-		{
-			Contents += Options.AppendPostSource();
-		}
 		// If preprocessed cache is enabled, debug dump occurs in the cook process rather than the workers, and
 		// in that case the env in Input.Environment has not been merged with the shared env. Do so here.
 		FShaderCompilerEnvironment MergedEnvironment(Input.Environment);
@@ -1953,11 +1947,26 @@ namespace UE::ShaderCompilerCommon
 		{
 			MergedEnvironment.Merge(*Input.SharedEnvironment);
 		}
+
+		FString Contents = MergedEnvironment.GetDefinitionsAsCommentedCode();
+
+		if (Options.AppendPreSource)
+		{
+			Contents += Options.AppendPreSource();
+		}
+
+		Contents += PreprocessedSource;
+
+		if (Options.AppendPostSource)
+		{
+			Contents += Options.AppendPostSource();
+		}
+
 		Contents += TEXT("\n");
 		Contents += SerializeEnvironmentToBase64(MergedEnvironment);
-		Contents += TEXT("#if 0 /*DIRECT COMPILE*/\n");
+		Contents += TEXT("/* DIRECT COMPILE\n");
 		Contents += CreateShaderCompilerWorkerDirectCommandLine(Input);
-		Contents += TEXT("\n#endif /*DIRECT COMPILE*/\n");
+		Contents += TEXT("\nDIRECT COMPILE */\n");
 		if (!Input.DebugDescription.IsEmpty())
 		{
 			Contents += TEXT("//");
diff --git a/Engine/Source/Developer/ShaderFormatVectorVM/Private/VectorVMShaderCompiler.cpp b/Engine/Source/Developer/ShaderFormatVectorVM/Private/VectorVMShaderCompiler.cpp
--- a/Engine/Source/Developer/ShaderFormatVectorVM/Private/VectorVMShaderCompiler.cpp
+++ b/Engine/Source/Developer/ShaderFormatVectorVM/Private/VectorVMShaderCompiler.cpp
@@ -29,12 +29,7 @@ DECLARE_CYCLE_STAT(TEXT("VectorVM - Compiler - CrossCompilerContextRun"), STAT_V
 bool PreprocessVectorVMShader(const FShaderCompilerInput& Input, const FShaderCompilerEnvironment& Environment, FShaderPreprocessOutput& Output)
 {
 	SCOPE_CYCLE_COUNTER(STAT_VectorVM_Compiler_CompileShader_VectorVMPreprocessShader);
-	// disable deprecation warnings due to default-constructed FShaderCompilerDefinitions in function args.
-	// can be removed once FShaderCompilerDefinitions moves to Internal
-	PRAGMA_DISABLE_DEPRECATION_WARNINGS 	
-	// Don't include shader definitions since it creates shader compilation errors.
-	return PreprocessShader(Output, Input, Environment, {}, EDumpShaderDefines::DontIncludeDefines);
-	PRAGMA_ENABLE_DEPRECATION_WARNINGS
+	return PreprocessShader(Output, Input, Environment);
 }
 
 bool CompileVectorVMShader(
diff --git a/Engine/Source/Developer/ShaderPreprocessor/Private/ShaderPreprocessor.cpp b/Engine/Source/Developer/ShaderPreprocessor/Private/ShaderPreprocessor.cpp
--- a/Engine/Source/Developer/ShaderPreprocessor/Private/ShaderPreprocessor.cpp
+++ b/Engine/Source/Developer/ShaderPreprocessor/Private/ShaderPreprocessor.cpp
@@ -846,13 +846,14 @@ bool PreprocessShader(
 	const FShaderCompilerInput& ShaderInput,
 	PRAGMA_DISABLE_DEPRECATION_WARNINGS		// FShaderCompilerDefinitions will be made internal in the future, marked deprecated until then
 	const FShaderCompilerDefinitions& AdditionalDefines,
+	EDumpShaderDefines DefinesPolicy
 	PRAGMA_ENABLE_DEPRECATION_WARNINGS
-	EDumpShaderDefines DefinesPolicy)
+)
 {
 	FShaderPreprocessOutput Output;
 	// when called via this overload, environment is assumed to be already merged in input struct
 	const FShaderCompilerEnvironment& Environment = ShaderInput.Environment;
-	bool bSucceeded = PreprocessShader(Output, ShaderInput, Environment, AdditionalDefines, DefinesPolicy);
+	bool bSucceeded = PreprocessShader(Output, ShaderInput, Environment, AdditionalDefines);
 
 	OutPreprocessedShader = MoveTemp(Output.EditSource());
 
@@ -878,9 +879,8 @@ bool PreprocessShader(
 	const FShaderCompilerInput& Input,
 	const FShaderCompilerEnvironment& Environment,
 	PRAGMA_DISABLE_DEPRECATION_WARNINGS		// FShaderCompilerDefinitions will be made internal in the future, marked deprecated until then
-	const FShaderCompilerDefinitions& AdditionalDefines,
+	const FShaderCompilerDefinitions& AdditionalDefines
 	PRAGMA_ENABLE_DEPRECATION_WARNINGS
-	EDumpShaderDefines DefinesPolicy
 )
 {
 	TRACE_CPUPROFILER_EVENT_SCOPE(PreprocessShader);
@@ -896,6 +896,12 @@ bool PreprocessShader(
 		{
 			// const_cast for compile environment; need to populate a subset of environment parameters from parsing comments in the preprocessed code
 			UE::ShaderCompilerCommon::SerializeEnvironmentFromBase64(const_cast<FShaderCompilerEnvironment&>(Input.Environment), Output.GetSource());
+
+			// strip comments from source when loading from a debug USF. some backends don't handle the comments that the debug dump inserts properly.
+			// this (currently) incurs an extra conversion (TCHAR -> ANSI -> TCHAR), but since this is only used in the debug path the perf hit is irrelevant
+			TArray<ANSICHAR> Stripped;
+			ShaderConvertAndStripComments(Output.GetSource(), Stripped);
+			Output.EditSource() = Stripped.GetData();
 		}
 
 		return bSuccess;
@@ -903,12 +909,6 @@ bool PreprocessShader(
 
 	check(CheckVirtualShaderFilePath(Input.VirtualSourceFilePath));
 
-	// List the defines used for compilation in the preprocessed shaders, especially to know which permutation vector this shader is.
-	if (DefinesPolicy == EDumpShaderDefines::AlwaysIncludeDefines || (DefinesPolicy == EDumpShaderDefines::DontCare && Input.DumpDebugInfoPath.Len() > 0))
-	{
-		FShaderPreprocessorUtilities::DumpShaderDefinesAsCommentedCode(Environment, &Output.EditSource());
-	}
-
 	bool bSuccess = InnerPreprocessShaderStb(Output, Input, Environment, AdditionalDefines);
 
 	if (bSuccess)
diff --git a/Engine/Source/Developer/ShaderPreprocessor/Public/ShaderPreprocessor.h b/Engine/Source/Developer/ShaderPreprocessor/Public/ShaderPreprocessor.h
--- a/Engine/Source/Developer/ShaderPreprocessor/Public/ShaderPreprocessor.h
+++ b/Engine/Source/Developer/ShaderPreprocessor/Public/ShaderPreprocessor.h
@@ -12,10 +12,9 @@ class FString;
 struct FShaderCompilerInput;
 struct FShaderCompilerOutput;
 
-/** Governs the behavior for adding shader defines to the preprocessed source. Can be helpful for the debugging, but makes the source unique
-    which can prevent efficient caching.
-  */
-enum class EDumpShaderDefines : uint8
+enum class 
+	UE_DEPRECATED(5.4, "EDumpShaderDefines is no longer used. Shader defines will always be included in preprocessed source debug dumps, but no longer explicitly added to the source by preprocessing.")
+	EDumpShaderDefines : uint8
 {
 	/** Will not be dumped unless Input.DumpDebugInfoPath is set */
 	DontCare,
@@ -25,15 +24,6 @@ enum class EDumpShaderDefines : uint8
 	AlwaysIncludeDefines
 };
 
-/**
- * Preprocess a shader.
- * @param OutPreprocessedShader - Upon return contains the preprocessed source code.
- * @param ShaderOutput - ShaderOutput to which errors can be added.
- * @param ShaderInput - The shader compiler input.
- * @param AdditionalDefines - Additional defines with which to preprocess the shader.
- * @param bShaderDumpDefinesAsCommentedCode - Whether to add shader definitions as comments.
- * @returns true if the shader is preprocessed without error.
- */
 UE_DEPRECATED(5.4, "Please use overload of PreprocessShader accepting a FShaderPreprocessOutput struct.")
 extern SHADERPREPROCESSOR_API bool PreprocessShader(
 	FString& OutPreprocessedShader,
@@ -41,8 +31,9 @@ extern SHADERPREPROCESSOR_API bool PreprocessShader(
 	const FShaderCompilerInput& ShaderInput,
 	PRAGMA_DISABLE_DEPRECATION_WARNINGS		// FShaderCompilerDefinitions will be made internal in the future, marked deprecated until then
 	const FShaderCompilerDefinitions& AdditionalDefines,
+	EDumpShaderDefines DefinesPolicy = EDumpShaderDefines::DontCare
 	PRAGMA_ENABLE_DEPRECATION_WARNINGS
-	EDumpShaderDefines DefinesPolicy = EDumpShaderDefines::DontCare);
+	);
 
 /**
  * Preprocess a shader.
@@ -52,7 +43,6 @@ extern SHADERPREPROCESSOR_API bool PreprocessShader(
  * (it is assumed this overload is called outside of the worker process which merges this in-place, so this merge step must be
  * performed by the caller)
  * @param AdditionalDefines - Additional defines with which to preprocess the shader.
- * @param bShaderDumpDefinesAsCommentedCode - Whether to add shader definitions as comments.
  * @returns true if the shader is preprocessed without error.
  */
 extern SHADERPREPROCESSOR_API bool PreprocessShader(
@@ -60,6 +50,21 @@ extern SHADERPREPROCESSOR_API bool PreprocessShader(
 	const FShaderCompilerInput& Input,
 	const FShaderCompilerEnvironment& MergedEnvironment,
 	PRAGMA_DISABLE_DEPRECATION_WARNINGS		// FShaderCompilerDefinitions will be made internal in the future, marked deprecated until then
-	const FShaderCompilerDefinitions& AdditionalDefines = FShaderCompilerDefinitions(),
+	const FShaderCompilerDefinitions& AdditionalDefines = FShaderCompilerDefinitions()
+	PRAGMA_ENABLE_DEPRECATION_WARNINGS
+);
+
+UE_DEPRECATED(5.4, "EDumpShaderDefines is no longer used. Shader defines will always be included in preprocessed source debug dumps, but no longer explicitly added to the source by preprocessing.")
+inline bool PreprocessShader(
+	FShaderPreprocessOutput& Output,
+	const FShaderCompilerInput& Input,
+	const FShaderCompilerEnvironment& MergedEnvironment,
+	PRAGMA_DISABLE_DEPRECATION_WARNINGS		// FShaderCompilerDefinitions will be made internal in the future, marked deprecated until then
+	const FShaderCompilerDefinitions& AdditionalDefines,
+	EDumpShaderDefines DefinesPolicy
 	PRAGMA_ENABLE_DEPRECATION_WARNINGS
-	EDumpShaderDefines DefinesPolicy = EDumpShaderDefines::DontCare);
+)
+{
+	return PreprocessShader(Output, Input, MergedEnvironment, AdditionalDefines);
+}
+
diff --git a/Engine/Source/Runtime/RenderCore/Private/ShaderCore.cpp b/Engine/Source/Runtime/RenderCore/Private/ShaderCore.cpp
--- a/Engine/Source/Runtime/RenderCore/Private/ShaderCore.cpp
+++ b/Engine/Source/Runtime/RenderCore/Private/ShaderCore.cpp
@@ -874,7 +874,7 @@ void FShaderCompilerEnvironment::Merge(const FShaderCompilerEnvironment& Other)
 	CompilerFlags.Append(Other.CompilerFlags);
 	ResourceTableMap.Append(Other.ResourceTableMap);
 	UniformBufferMap.Append(Other.UniformBufferMap);
-	PRAGMA_DISABLE_DEPRECATION_WARNINGS		// FShaderCompilerDefinitions will be made internal in the future, marked deprecated until then
+	PRAGMA_DISABLE_DEPRECATION_WARNINGS	// FShaderCompilerDefinitions will be made internal in the future, marked deprecated until then
 	Definitions->Merge(*Other.Definitions);
 	PRAGMA_ENABLE_DEPRECATION_WARNINGS
 	CompileArgs.Append(Other.CompileArgs);
@@ -882,6 +882,27 @@ void FShaderCompilerEnvironment::Merge(const FShaderCompilerEnvironment& Other)
 	FullPrecisionInPS |= Other.FullPrecisionInPS;
 }
 
+FString FShaderCompilerEnvironment::GetDefinitionsAsCommentedCode() const
+{
+	TArray<FString> DefinesLines;
+	PRAGMA_DISABLE_DEPRECATION_WARNINGS // FShaderCompilerDefinitions will be made internal in the future, marked deprecated until then
+	DefinesLines.Reserve(Definitions->Num());
+	for (FShaderCompilerDefinitions::FConstIterator DefineIt(*Definitions); DefineIt; ++DefineIt)
+	PRAGMA_ENABLE_DEPRECATION_WARNINGS
+	{
+		DefinesLines.Add(FString::Printf(TEXT("// #define %s %s\n"), DefineIt.Key(), DefineIt.Value()));
+	}
+	DefinesLines.Sort();
+
+	FString Defines;
+	for (const FString& DefineLine : DefinesLines)
+	{
+		Defines += DefineLine;
+	}
+
+	return MakeInjectedShaderCodeBlock(TEXT("DumpShaderDefinesAsCommentedCode"), Defines);
+}
+
 PRAGMA_DISABLE_DEPRECATION_WARNINGS		// FShaderCompilerDefinitions will be made internal in the future, marked deprecated until then
 
 // Pass through functions to definitions
diff --git a/Engine/Source/Runtime/RenderCore/Public/ShaderCore.h b/Engine/Source/Runtime/RenderCore/Public/ShaderCore.h
--- a/Engine/Source/Runtime/RenderCore/Public/ShaderCore.h
+++ b/Engine/Source/Runtime/RenderCore/Public/ShaderCore.h
@@ -650,6 +650,8 @@ struct FShaderCompilerEnvironment
 	
 	RENDERCORE_API void Merge(const FShaderCompilerEnvironment& Other);
 
+	RENDERCORE_API FString GetDefinitionsAsCommentedCode() const;
+
 private:
 
 	friend class FShaderCompileUtilities;
