[shaders] further preprocessing cleanup
- move sequence of preprocessing steps out of ShaderPreprocessor module and into UE::ShaderCompilerCommon::ExecuteShaderPreprocessingSteps; the former is now explicitly just the low-level preprocessor lib
- add an implementation of PreprocessShader in FBaseShaderFormat so backends which have no custom code to execute as part of preprocessing can just automatically inherit this implementation, and fix up such backends to eliminate now-unnecessary overrides

diff --git a/Engine/Source/Developer/Apple/MetalShaderFormat/Private/MetalShaderCompiler.cpp b/Engine/Source/Developer/Apple/MetalShaderFormat/Private/MetalShaderCompiler.cpp
--- a/Engine/Source/Developer/Apple/MetalShaderFormat/Private/MetalShaderCompiler.cpp
+++ b/Engine/Source/Developer/Apple/MetalShaderFormat/Private/MetalShaderCompiler.cpp
@@ -763,7 +763,7 @@ bool PreprocessMetalShader(const FShaderCompilerInput& Input, const FShaderCompi
 		return false;
 	}
 
-	return PreprocessShader(PreprocessOutput, Input, Environment);
+	return UE::ShaderCompilerCommon::ExecuteShaderPreprocessingSteps(PreprocessOutput, Input, Environment);
 }
 
 void CompileMetalShader(const FShaderCompilerInput& Input, const FString& InPreprocessedSource, FShaderCompilerOutput& Output)
diff --git a/Engine/Source/Developer/ShaderCompilerCommon/Private/ShaderCompilerCommon.cpp b/Engine/Source/Developer/ShaderCompilerCommon/Private/ShaderCompilerCommon.cpp
--- a/Engine/Source/Developer/ShaderCompilerCommon/Private/ShaderCompilerCommon.cpp
+++ b/Engine/Source/Developer/ShaderCompilerCommon/Private/ShaderCompilerCommon.cpp
@@ -12,6 +12,7 @@
 #include "Serialization/MemoryReader.h"
 #include "Serialization/MemoryWriter.h"
 #include "String/RemoveFrom.h"
+#include "ShaderPreprocessor.h"
 #include "ShaderPreprocessTypes.h"
 #include "ShaderSymbolExport.h"
 #include "ShaderMinifier.h"
@@ -1806,11 +1807,69 @@ const FString GetDebugFileName(
 
 namespace UE::ShaderCompilerCommon
 {
+	bool ExecuteShaderPreprocessingSteps(
+		FShaderPreprocessOutput& PreprocessOutput,
+		const FShaderCompilerInput& Input,
+		const FShaderCompilerEnvironment& Environment,
+		PRAGMA_DISABLE_DEPRECATION_WARNINGS
+		const FShaderCompilerDefinitions& AdditionalDefines
+		PRAGMA_ENABLE_DEPRECATION_WARNINGS
+		)
+	{
+		TRACE_CPUPROFILER_EVENT_SCOPE(FBaseShaderFormat_PreprocessShader);
+		PreprocessOutput.EditSource().Empty();
+
+		if (EnumHasAnyFlags(Input.DebugInfoFlags, EShaderDebugInfoFlags::CompileFromDebugUSF))
+		{
+			// the "VirtualSourceFilePath" given is actually an absolute path to a dumped debug USF file; load it directly.
+			// this occurs when running SCW in "direct compile" mode; this file will already be preprocessed.
+			bool bSuccess = FFileHelper::LoadFileToString(PreprocessOutput.EditSource(), *Input.VirtualSourceFilePath);
+
+			if (bSuccess)
+			{
+				// const_cast for compile environment; need to populate a subset of environment parameters from parsing comments in the preprocessed code
+				UE::ShaderCompilerCommon::SerializeEnvironmentFromBase64(const_cast<FShaderCompilerEnvironment&>(Input.Environment), PreprocessOutput.GetSource());
+
+				// strip comments from source when loading from a debug USF. some backends don't handle the comments that the debug dump inserts properly.
+				// this (currently) incurs an extra conversion (TCHAR -> ANSI -> TCHAR), but since this is only used in the debug path the perf hit is irrelevant
+				TArray<ANSICHAR> Stripped;
+				ShaderConvertAndStripComments(PreprocessOutput.GetSource(), Stripped);
+				PreprocessOutput.EditSource() = Stripped.GetData();
+			}
+
+			return bSuccess;
+		}
+
+		check(CheckVirtualShaderFilePath(Input.VirtualSourceFilePath));
+
+		bool bSuccess = ::PreprocessShader(PreprocessOutput, Input, Environment, AdditionalDefines);
+		if (bSuccess)
+		{
+			CleanupUniformBufferCode(Environment, PreprocessOutput.EditSource());
+
+			if (Input.Environment.CompilerFlags.Contains(CFLAG_RemoveDeadCode))
+			{
+				const TArray<FStringView> RequiredSymbols(MakeArrayView(Input.RequiredSymbols));
+				UE::ShaderCompilerCommon::RemoveDeadCode(PreprocessOutput.EditSource(), Input.EntryPointName, RequiredSymbols, PreprocessOutput.EditErrors());
+			}
+		}
+
+		return bSuccess;
+	}
+
 	FString FDebugShaderDataOptions::GetDebugShaderPath(const FShaderCompilerInput& Input) const
 	{
 		return GetDebugFileName(Input, *this, OverrideBaseFilename);
 	}
 
+	bool FBaseShaderFormat::PreprocessShader(
+		const FShaderCompilerInput& Input,
+		const FShaderCompilerEnvironment& Environment,
+		FShaderPreprocessOutput& PreprocessOutput) const
+	{
+		return ExecuteShaderPreprocessingSteps(PreprocessOutput, Input, Environment);
+	}
+
 	void FBaseShaderFormat::OutputDebugData(
 		const FShaderCompilerInput& Input,
 		const FShaderPreprocessOutput& PreprocessOutput,
diff --git a/Engine/Source/Developer/ShaderCompilerCommon/Public/ShaderCompilerCommon.h b/Engine/Source/Developer/ShaderCompilerCommon/Public/ShaderCompilerCommon.h
--- a/Engine/Source/Developer/ShaderCompilerCommon/Public/ShaderCompilerCommon.h
+++ b/Engine/Source/Developer/ShaderCompilerCommon/Public/ShaderCompilerCommon.h
@@ -64,6 +64,15 @@ extern SHADERCOMPILERCOMMON_API int16 GetNumUniformBuffersUsed(const FShaderComp
 
 namespace UE::ShaderCompilerCommon
 {
+	extern SHADERCOMPILERCOMMON_API bool ExecuteShaderPreprocessingSteps(
+		FShaderPreprocessOutput& PreprocessOutput,
+		const FShaderCompilerInput& Input,
+		const FShaderCompilerEnvironment& Environment,
+		PRAGMA_DISABLE_DEPRECATION_WARNINGS		// FShaderCompilerDefinitions will be made internal in the future, marked deprecated until then
+		const FShaderCompilerDefinitions& AdditionalDefines = FShaderCompilerDefinitions()
+		PRAGMA_ENABLE_DEPRECATION_WARNINGS
+		);
+
 	extern SHADERCOMPILERCOMMON_API FStringView          RemoveConstantBufferPrefix(FStringView InName);
 	extern SHADERCOMPILERCOMMON_API FString              RemoveConstantBufferPrefix(const FString& InName);
 
@@ -170,6 +179,11 @@ namespace UE::ShaderCompilerCommon
 	class FBaseShaderFormat : public IShaderFormat
 	{
 	public:
+		virtual SHADERCOMPILERCOMMON_API bool PreprocessShader(
+			const FShaderCompilerInput& Input,
+			const FShaderCompilerEnvironment& Environment,
+			FShaderPreprocessOutput& PreprocessOutput) const override;
+
 		virtual SHADERCOMPILERCOMMON_API void OutputDebugData(
 			const FShaderCompilerInput& Input,
 			const FShaderPreprocessOutput& PreprocessOutput,
diff --git a/Engine/Source/Developer/ShaderCompilerCommon/ShaderCompilerCommon.Build.cs b/Engine/Source/Developer/ShaderCompilerCommon/ShaderCompilerCommon.Build.cs
--- a/Engine/Source/Developer/ShaderCompilerCommon/ShaderCompilerCommon.Build.cs
+++ b/Engine/Source/Developer/ShaderCompilerCommon/ShaderCompilerCommon.Build.cs
@@ -11,6 +11,7 @@ public class ShaderCompilerCommon : ModuleRules
 			new string[] {
 				"Core",
 				"RenderCore",
+				"ShaderPreprocessor",
 				"TargetPlatform",
 			}
 			);
diff --git a/Engine/Source/Developer/ShaderFormatOpenGL/Private/ShaderFormatOpenGL.cpp b/Engine/Source/Developer/ShaderFormatOpenGL/Private/ShaderFormatOpenGL.cpp
--- a/Engine/Source/Developer/ShaderFormatOpenGL/Private/ShaderFormatOpenGL.cpp
+++ b/Engine/Source/Developer/ShaderFormatOpenGL/Private/ShaderFormatOpenGL.cpp
@@ -106,12 +106,6 @@ public:
 		}
 	}
 
-	virtual bool PreprocessShader(const FShaderCompilerInput& Input, const FShaderCompilerEnvironment& Environment, FShaderPreprocessOutput& PreprocessOutput) const override
-	{
-		CheckFormat(Input.ShaderFormat);
-		return ::PreprocessShader(PreprocessOutput, Input, Environment);
-	}
-
 	virtual void CompilePreprocessedShader(const FShaderCompilerInput& Input, const FShaderPreprocessOutput& PreprocessOutput, FShaderCompilerOutput& Output, const FString& WorkingDirectory) const override
 	{
 		CheckFormat(Input.ShaderFormat);
diff --git a/Engine/Source/Developer/ShaderFormatVectorVM/Private/ShaderFormatVectorVM.cpp b/Engine/Source/Developer/ShaderFormatVectorVM/Private/ShaderFormatVectorVM.cpp
--- a/Engine/Source/Developer/ShaderFormatVectorVM/Private/ShaderFormatVectorVM.cpp
+++ b/Engine/Source/Developer/ShaderFormatVectorVM/Private/ShaderFormatVectorVM.cpp
@@ -8,14 +8,11 @@
 #include "Interfaces/IShaderFormat.h"
 #include "Interfaces/IShaderFormatModule.h"
 #include "hlslcc.h"
+#include "ShaderCompilerCommon.h"
 #include "ShaderCompilerCore.h"
 #include "ShaderCore.h"
 
 
-extern bool PreprocessVectorVMShader(
-	const FShaderCompilerInput& Input,
-	const FShaderCompilerEnvironment& Environment,
-	class FShaderPreprocessOutput& PreprocessOutput);
 
 extern bool CompileVectorVMShader(
 	const FShaderCompilerInput& Input,
@@ -30,7 +27,7 @@ extern void OutputVectorVMDebugData(
 
 static FName NAME_VVM_1_0(TEXT("VVM_1_0"));
 
-class FShaderFormatVectorVM : public IShaderFormat
+class FShaderFormatVectorVM : public UE::ShaderCompilerCommon::FBaseShaderFormat
 {
 	enum class VectorVMFormats : uint8
 	{
@@ -74,12 +71,6 @@ public:
 		Input.Environment.CompilerFlags.Remove(CFLAG_RemoveDeadCode);
 	}
 
-	virtual bool PreprocessShader(const FShaderCompilerInput& Input, const FShaderCompilerEnvironment& Environment, FShaderPreprocessOutput& PreprocessOutput) const override
-	{
-		CheckFormat(Input.ShaderFormat);
-		return PreprocessVectorVMShader(Input, Environment, PreprocessOutput);
-	}
-
 	virtual void CompilePreprocessedShader(const FShaderCompilerInput& Input, const FShaderPreprocessOutput& PreprocessOutput, FShaderCompilerOutput& Output, const FString& WorkingDirectory) const override
 	{
 		CheckFormat(Input.ShaderFormat);
diff --git a/Engine/Source/Developer/ShaderFormatVectorVM/Private/VectorVMShaderCompiler.cpp b/Engine/Source/Developer/ShaderFormatVectorVM/Private/VectorVMShaderCompiler.cpp
--- a/Engine/Source/Developer/ShaderFormatVectorVM/Private/VectorVMShaderCompiler.cpp
+++ b/Engine/Source/Developer/ShaderFormatVectorVM/Private/VectorVMShaderCompiler.cpp
@@ -26,12 +26,6 @@ DECLARE_CYCLE_STAT(TEXT("VectorVM - Compiler - CompileShader_VectorVM"), STAT_Ve
 DECLARE_CYCLE_STAT(TEXT("VectorVM - Compiler - PreprocessShader"), STAT_VectorVM_Compiler_CompileShader_VectorVMPreprocessShader, STATGROUP_VectorVM);
 DECLARE_CYCLE_STAT(TEXT("VectorVM - Compiler - CrossCompilerContextRun"), STAT_VectorVM_Compiler_CompileShader_CrossCompilerContextRun, STATGROUP_VectorVM);
 
-bool PreprocessVectorVMShader(const FShaderCompilerInput& Input, const FShaderCompilerEnvironment& Environment, FShaderPreprocessOutput& Output)
-{
-	SCOPE_CYCLE_COUNTER(STAT_VectorVM_Compiler_CompileShader_VectorVMPreprocessShader);
-	return PreprocessShader(Output, Input, Environment);
-}
-
 bool CompileVectorVMShader(
 	const FShaderCompilerInput& Input,
 	const FShaderPreprocessOutput& PreprocessOutput,
@@ -151,7 +145,7 @@ bool TestCompileVectorVMShader(
 	bool bSkipBackendOptimizations)
 {
 	FShaderPreprocessOutput PreprocessOutput;
-	if (!PreprocessVectorVMShader(Input, Input.Environment, PreprocessOutput))
+	if (!UE::ShaderCompilerCommon::ExecuteShaderPreprocessingSteps(PreprocessOutput, Input, Input.Environment))
 	{
 		if (PreprocessOutput.GetErrors().Num() != 0)
 		{
diff --git a/Engine/Source/Developer/ShaderPreprocessor/Private/ShaderPreprocessor.cpp b/Engine/Source/Developer/ShaderPreprocessor/Private/ShaderPreprocessor.cpp
--- a/Engine/Source/Developer/ShaderPreprocessor/Private/ShaderPreprocessor.cpp
+++ b/Engine/Source/Developer/ShaderPreprocessor/Private/ShaderPreprocessor.cpp
@@ -6,7 +6,6 @@
 #include "Misc/ScopeLock.h"
 #include "Modules/ModuleManager.h"
 #include "PreprocessorPrivate.h"
-#include "ShaderCompilerCommon.h"
 #include "ShaderCompilerDefinitions.h"
 #include "stb_preprocess/preprocessor.h"
 #include "stb_preprocess/stb_alloc.h"
@@ -76,25 +75,6 @@ static void AddStbDefines(stb_arena* MacroArena, macro_definition**& StbDefines,
 class FShaderPreprocessorUtilities
 {
 public:
-	static void DumpShaderDefinesAsCommentedCode(const FShaderCompilerEnvironment& Environment, FString* OutDefines)
-	{
-		TArray<FString> DefinesLines;
-		DefinesLines.Reserve(Environment.Definitions->Num());
-		for (FShaderCompilerDefinitions::FConstIterator DefineIt(*Environment.Definitions); DefineIt; ++DefineIt)
-		{
-			DefinesLines.Add(FString::Printf(TEXT("// #define %s %s\n"), DefineIt.Key(), DefineIt.Value()));
-		}
-		DefinesLines.Sort();
-
-		FString Defines;
-		for (const FString& DefineLine : DefinesLines)
-		{
-			Defines += DefineLine;
-		}
-
-		*OutDefines += MakeInjectedShaderCodeBlock(TEXT("DumpShaderDefinesAsCommentedCode"), Defines);
-	}
-
 	static void PopulateDefines(const FShaderCompilerEnvironment& Environment, const FShaderCompilerDefinitions& AdditionalDefines, stb_arena* MacroArena, macro_definition**& OutDefines)
 	{
 		arrsetcap(OutDefines, Environment.Definitions->Num() + AdditionalDefines.Num());
@@ -691,15 +671,33 @@ static void AddStbDefines(stb_arena* MacroArena, macro_definition**& StbDefines,
 	}
 }
 
-bool InnerPreprocessShaderStb(
+PRAGMA_ENABLE_DEPRECATION_WARNINGS
+
+/**
+ * Preprocess a shader.
+ * @param OutPreprocessedShader - Upon return contains the preprocessed source code.
+ * @param ShaderOutput - ShaderOutput to which errors can be added.
+ * @param ShaderInput - The shader compiler input.
+ * @param AdditionalDefines - Additional defines with which to preprocess the shader.
+ * @param DefinesPolicy - Whether to add shader definitions as comments.
+ * @returns true if the shader is preprocessed without error.
+ */
+bool PreprocessShader(
 	FShaderPreprocessOutput& Output,
 	const FShaderCompilerInput& Input,
 	const FShaderCompilerEnvironment& Environment,
+	PRAGMA_DISABLE_DEPRECATION_WARNINGS		// FShaderCompilerDefinitions will be made internal in the future, marked deprecated until then
 	const FShaderCompilerDefinitions& AdditionalDefines
+	PRAGMA_ENABLE_DEPRECATION_WARNINGS
 )
 {
+	TRACE_CPUPROFILER_EVENT_SCOPE(PreprocessShader);
+
 	stb_arena MacroArena = { 0 };
 	macro_definition** StbDefines = nullptr;
+	
+	PRAGMA_DISABLE_DEPRECATION_WARNINGS
+	
 	FShaderPreprocessorUtilities::PopulateDefines(Environment, AdditionalDefines, &MacroArena, StbDefines);
 
 	// The substitution text generated by custom macros gets run through the preprocessor afterwards, but in some cases we want to
@@ -863,64 +861,4 @@ bool PreprocessShader(
 		ShaderOutput.Errors.Add(MoveTemp(Error));
 	}
 	return bSucceeded;
-}
-
-/**
- * Preprocess a shader.
- * @param OutPreprocessedShader - Upon return contains the preprocessed source code.
- * @param ShaderOutput - ShaderOutput to which errors can be added.
- * @param ShaderInput - The shader compiler input.
- * @param AdditionalDefines - Additional defines with which to preprocess the shader.
- * @param DefinesPolicy - Whether to add shader definitions as comments.
- * @returns true if the shader is preprocessed without error.
- */
-bool PreprocessShader(
-	FShaderPreprocessOutput& Output,
-	const FShaderCompilerInput& Input,
-	const FShaderCompilerEnvironment& Environment,
-	PRAGMA_DISABLE_DEPRECATION_WARNINGS		// FShaderCompilerDefinitions will be made internal in the future, marked deprecated until then
-	const FShaderCompilerDefinitions& AdditionalDefines
-	PRAGMA_ENABLE_DEPRECATION_WARNINGS
-)
-{
-	TRACE_CPUPROFILER_EVENT_SCOPE(PreprocessShader);
-	Output.EditSource().Empty();
-
-	if (EnumHasAnyFlags(Input.DebugInfoFlags, EShaderDebugInfoFlags::CompileFromDebugUSF))
-	{
-		// the "VirtualSourceFilePath" given is actually an absolute path to a dumped debug USF file; load it directly.
-		// this occurs when running SCW in "direct compile" mode; this file will already be preprocessed.
-		bool bSuccess = FFileHelper::LoadFileToString(Output.EditSource(), *Input.VirtualSourceFilePath);
-
-		if (bSuccess)
-		{
-			// const_cast for compile environment; need to populate a subset of environment parameters from parsing comments in the preprocessed code
-			UE::ShaderCompilerCommon::SerializeEnvironmentFromBase64(const_cast<FShaderCompilerEnvironment&>(Input.Environment), Output.GetSource());
-
-			// strip comments from source when loading from a debug USF. some backends don't handle the comments that the debug dump inserts properly.
-			// this (currently) incurs an extra conversion (TCHAR -> ANSI -> TCHAR), but since this is only used in the debug path the perf hit is irrelevant
-			TArray<ANSICHAR> Stripped;
-			ShaderConvertAndStripComments(Output.GetSource(), Stripped);
-			Output.EditSource() = Stripped.GetData();
-		}
-
-		return bSuccess;
-	}
-
-	check(CheckVirtualShaderFilePath(Input.VirtualSourceFilePath));
-
-	bool bSuccess = InnerPreprocessShaderStb(Output, Input, Environment, AdditionalDefines);
-
-	if (bSuccess)
-	{
-		CleanupUniformBufferCode(Environment, Output.EditSource());
-
-		if (Input.Environment.CompilerFlags.Contains(CFLAG_RemoveDeadCode))
-		{
-			const TArray<FStringView> RequiredSymbols(MakeArrayView(Input.RequiredSymbols));
-			UE::ShaderCompilerCommon::RemoveDeadCode(Output.EditSource(), Input.EntryPointName, RequiredSymbols, Output.EditErrors());
-		}
-	}
-
-	return bSuccess;
-}
+}
\ No newline at end of file
diff --git a/Engine/Source/Developer/ShaderPreprocessor/ShaderPreprocessor.Build.cs b/Engine/Source/Developer/ShaderPreprocessor/ShaderPreprocessor.Build.cs
--- a/Engine/Source/Developer/ShaderPreprocessor/ShaderPreprocessor.Build.cs
+++ b/Engine/Source/Developer/ShaderPreprocessor/ShaderPreprocessor.Build.cs
@@ -10,8 +10,6 @@ public class ShaderPreprocessor : ModuleRules
 			new string[] {
 				"Core",
 				"RenderCore",
-				"ShaderCompilerCommon",
-				"TargetPlatform",
 			}
 			);
 
diff --git a/Engine/Source/Developer/VulkanShaderFormat/Private/VulkanShaderFormat.cpp b/Engine/Source/Developer/VulkanShaderFormat/Private/VulkanShaderFormat.cpp
--- a/Engine/Source/Developer/VulkanShaderFormat/Private/VulkanShaderFormat.cpp
+++ b/Engine/Source/Developer/VulkanShaderFormat/Private/VulkanShaderFormat.cpp
@@ -34,7 +34,7 @@ static const FGuid UE_SHADER_VULKAN_ES3_1_VER = FGuid("B84F72C8-3ECD-411E-993C-D
 static const FGuid UE_SHADER_VULKAN_SM5_VER = FGuid("0715D8EE-9907-4A25-93AD-A3902C8E069A");
 static const FGuid UE_SHADER_VULKAN_SM6_VER = FGuid("C5161730-83C6-40AF-A990-78CD4C1581DB");
 
-class FShaderFormatVulkan : public IShaderFormat
+class FShaderFormatVulkan : public UE::ShaderCompilerCommon::FBaseShaderFormat
 {
 	FGuid InternalGetVersion(FName Format) const
 	{
@@ -93,11 +93,6 @@ public:
 		ModifyVulkanCompilerInput(Input);
 	}
 
-	virtual bool PreprocessShader(const FShaderCompilerInput& Input, const FShaderCompilerEnvironment& Environment, FShaderPreprocessOutput& PreprocessOutput) const
-	{
-		return ::PreprocessShader(PreprocessOutput, Input, Environment);
-	}
-
 	virtual void CompilePreprocessedShader(const FShaderCompilerInput& Input, const FShaderPreprocessOutput& PreprocessOutput, FShaderCompilerOutput& Output,const FString& WorkingDirectory) const override
 	{
 		CompileVulkanShader(Input, PreprocessOutput.GetSource(), Output, WorkingDirectory);
diff --git a/Engine/Source/Developer/Windows/ShaderFormatD3D/Private/ShaderFormatD3D.cpp b/Engine/Source/Developer/Windows/ShaderFormatD3D/Private/ShaderFormatD3D.cpp
--- a/Engine/Source/Developer/Windows/ShaderFormatD3D/Private/ShaderFormatD3D.cpp
+++ b/Engine/Source/Developer/Windows/ShaderFormatD3D/Private/ShaderFormatD3D.cpp
@@ -171,11 +171,6 @@ public:
 	}
 #endif
 
-	virtual bool PreprocessShader(const FShaderCompilerInput& Input, const FShaderCompilerEnvironment& MergedEnvironment, FShaderPreprocessOutput& PreprocessOutput) const
-	{
-		return ::PreprocessShader(PreprocessOutput, Input, MergedEnvironment);
-	}
-
 	virtual void CompilePreprocessedShader(
 		const FShaderCompilerInput& Input, 
 		const FShaderPreprocessOutput& PreprocessOutput, 
