Shader preprocessor:  Fix for ASAN failure -- need to add SSE padding to substitution text generated for TEXT macros.

#rnx

diff --git a/Engine/Source/Developer/ShaderPreprocessor/Private/ShaderPreprocessor.cpp b/Engine/Source/Developer/ShaderPreprocessor/Private/ShaderPreprocessor.cpp
--- a/Engine/Source/Developer/ShaderPreprocessor/Private/ShaderPreprocessor.cpp
+++ b/Engine/Source/Developer/ShaderPreprocessor/Private/ShaderPreprocessor.cpp
@@ -326,6 +326,19 @@ static void CopyStringToAnsiCharArray(const TCHAR* Text, int32 TextLen, TArray<A
 	*OutData = 0;
 }
 
+// Adds 16 bytes of zeroes at end, to allow SSE reads at the end of the buffer without reading past the end of the heap allocation
+static void CopyStringToAnsiCharArraySSEPadded(const TCHAR* Text, int32 TextLen, TArray<ANSICHAR>& Out)
+{
+	constexpr int32 SSEPadding = 16;
+	Out.SetNumUninitialized(TextLen + SSEPadding);
+	ANSICHAR* OutData = Out.GetData();
+	for (int32 CharIndex = 0; CharIndex < TextLen; CharIndex++, OutData++, Text++)
+	{
+		*OutData = (ANSICHAR)*Text;
+	}
+	FMemory::Memset(OutData, 0, SSEPadding * sizeof(ANSICHAR));
+}
+
 static const ANSICHAR* StbResolveInclude(const ANSICHAR* PathInSource, uint32 PathLen, const ANSICHAR* ParentPathAnsi, void* RawContext)
 {
 	FStbPreprocessContext& Context = *reinterpret_cast<FStbPreprocessContext*>(RawContext);
@@ -533,11 +546,11 @@ static const char* StbCustomMacroBegin(const char* OriginalText, void* RawContex
 	Entry.Hash = CityHash32((const char*)Entry.SourceText.GetCharArray().GetData(), sizeof(FString::ElementType) * Entry.SourceText.Len());
 	Context.TextGlobalCount += Entry.ConvertedText.Len();
 
-	// Generate substitution string
+	// Generate substitution string -- need SSE padding on any text handled by the preprocessor
 	if (Entry.bIsAssert)
 	{
 		const FString HashString = FString::Printf(TEXT("%u"), Entry.Hash);
-		CopyStringToAnsiCharArray(*HashString, HashString.Len(), Context.TextMacroSubstituted);
+		CopyStringToAnsiCharArraySSEPadded(*HashString, HashString.Len(), Context.TextMacroSubstituted);
 	}
 	else
 	{
@@ -545,7 +558,7 @@ static const char* StbCustomMacroBegin(const char* OriginalText, void* RawContex
 		const FString InitHashEnd(TEXT(")"));
 
 		const FString HashText = InitHashBegin + FString::FromInt(EntryIndex) + InitHashEnd;
-		CopyStringToAnsiCharArray(*HashText, HashText.Len(), Context.TextMacroSubstituted);
+		CopyStringToAnsiCharArraySSEPadded(*HashText, HashText.Len(), Context.TextMacroSubstituted);
 	}
 
 	return Context.TextMacroSubstituted.GetData();
diff --git a/Engine/Source/Developer/ShaderPreprocessor/Private/stb_preprocess/preprocessor.h b/Engine/Source/Developer/ShaderPreprocessor/Private/stb_preprocess/preprocessor.h
--- a/Engine/Source/Developer/ShaderPreprocessor/Private/stb_preprocess/preprocessor.h
+++ b/Engine/Source/Developer/ShaderPreprocessor/Private/stb_preprocess/preprocessor.h
@@ -75,15 +75,15 @@ typedef void (*freefile_callback_func)(const char* filename, const char* loaded_
 // note: returned resolved paths must remain valid for the lifetime of a single preprocessor execution
 typedef const char* (*resolveinclude_callback_func)(const char* path, unsigned int path_len, const char* parent, void* custom_context);
 
-// Callback functions for custom macros.  Custom macros return their own substitution text, which is then further
-// preprocessed.  The text passed to the callback includes the macro identifier and its arguments, with backslashes
-// removed.  A function is called when entering the macro, and again when complete, allowing stateful logic
-// based on whether parsing is inside a certain custom macro.  Any custom macro will automatically be disabled
-// inside itself, with the idea that this allows a singleton buffer for the substitution text, as there's no
-// concern about nested calls.  If text is dynamically allocated, the caller is responsible for freeing it.  It's
-// also valid to return the original text passed in as the substitution text, if you just want to set some state
-// inside the context, and not actually change the macro text -- since the macro is disabled, it will then be
-// echoed as non-macro text.
+// Callback functions for custom macros.  Custom macros return their own substitution text, which is then further preprocessed.
+// The substitution text must include 15 bytes of padding past the null terminator, as the preprocessor uses SSE reads which may
+// read past the null terminator.  The text passed to the callback includes the macro identifier and its arguments, with backslashes
+// removed.  A function is called when entering the macro, and again when complete, allowing stateful logic based on whether
+// parsing is inside a certain custom macro.  Any custom macro will automatically be disabled inside itself, with the idea that
+// this allows a singleton buffer for the substitution text, as there's no concern about nested calls.  If text is dynamically
+// allocated, the caller is responsible for freeing it.  It's also valid to return the original text passed in as the substitution
+// text, if you just want to set some state inside the context, and not actually change the macro text -- since the macro is
+// disabled, it will then be echoed as non-macro text.
 typedef const char* (*custommacro_begin_callback_func)(const char* original_text, void* custom_context);
 typedef void (*custommacro_end_callback_func)(const char* original_text, void* custom_context, const char* substitution_text);
 
