[ShaderPreprocessor]
* Reverted change that should never have checked in that disabled sse opts

diff --git a/Engine/Source/Developer/ShaderPreprocessor/Private/stb_preprocess/preprocessor.c b/Engine/Source/Developer/ShaderPreprocessor/Private/stb_preprocess/preprocessor.c
--- a/Engine/Source/Developer/ShaderPreprocessor/Private/stb_preprocess/preprocessor.c
+++ b/Engine/Source/Developer/ShaderPreprocessor/Private/stb_preprocess/preprocessor.c
@@ -27,7 +27,7 @@
 #endif
 
 // Unreal 5.2+ assumes SSE 4.2 is available on x64 processors
-#if 0//(defined(_M_X64) || defined(__amd64__) || defined(__x86_64__)) && !defined(_M_ARM64EC)
+#if (defined(_M_X64) || defined(__amd64__) || defined(__x86_64__)) && !defined(_M_ARM64EC)
 #include <immintrin.h>
 #define PREPROCESSOR_USE_SSE4_2 1
 #define SSE_READ_PADDING 16
