[shaders/stb_preprocess] stb preprocessor bug fixes:
- fix line number mismatches caused due to multiline macro arguments by keeping leading/trailing whitespaces in cases of multiple lines encountered in a single macro argument (single line macro args still have whitespace trimmed as this is needed for token pasting)
- avoid early out preventing parsing of "defined" builtin macro outside of #if conditions; this allows using defined in #define directives (among other places)

diff --git a/Engine/Source/Developer/ShaderPreprocessor/Private/stb_preprocess/preprocessor.c b/Engine/Source/Developer/ShaderPreprocessor/Private/stb_preprocess/preprocessor.c
--- a/Engine/Source/Developer/ShaderPreprocessor/Private/stb_preprocess/preprocessor.c
+++ b/Engine/Source/Developer/ShaderPreprocessor/Private/stb_preprocess/preprocessor.c
@@ -2907,7 +2907,7 @@ static const char* copy_argument(const char* text, int* line_number, char** p_ou
 				if (p[-1] == '\\')
 					arrsetlen(out, arrlen(out) - 1);  // undo output of '\'
 				else
-					arrput(out, ' ');
+					arrput(out, '\n');
 				p += newline_char_count(p[0], p[1]);
 				++*line_number;
 				break;
@@ -2974,7 +2974,7 @@ static const char* copy_argument(const char* text, int* line_number, char** p_ou
 					if (*p == 0)
 						return p;
 					p += newline_char_count(p[0], p[1]);
-					arrput(out, ' ');
+					arrput(out, '\n');
 					++*line_number;
 				}
 				else if (p[1] == '*')
@@ -2986,6 +2986,7 @@ static const char* copy_argument(const char* text, int* line_number, char** p_ou
 						{
 							++*line_number;
 							p += newline_char_count(p[0], p[1]);
+							arrput(out, '\n');
 						}
 						else
 						{
@@ -3091,9 +3092,6 @@ static void maybe_expand_macro(parse_state* cs, struct macro_definition* pending
 				int parenthesized = 0;
 				size_t len;
 
-				if (in_macro_expansion != IN_MACRO_if_condition)
-					goto not_macro;
-
 				// copy through the following token as well
 				while (char_is_whitespace(*p))
 					++p;
@@ -3273,6 +3271,7 @@ static void maybe_expand_macro(parse_state* cs, struct macro_definition* pending
 				// Set source offset to the end of the macro
 				cs->src_offset = p - cs->src;
 				cs->src_line_number += arg_newlines;
+				cs->dest_line_number += arg_newlines;
 				break;
 			}
 		}
@@ -3342,8 +3341,9 @@ static void maybe_expand_macro(parse_state* cs, struct macro_definition* pending
 		{
 			char* copy = argument_buffer;
 			int arg_newlines = 0;
-
-			p = preprocessor_skip_whitespace(p, &cs->src_line_number);
+			const char* leading_ws = p;
+			p = preprocessor_skip_whitespace(p, &arg_newlines);
+			int leading_ws_size = p - leading_ws;
 
 			s = copy_argument(p, &arg_newlines, &copy);
 
@@ -3386,18 +3386,26 @@ static void maybe_expand_macro(parse_state* cs, struct macro_definition* pending
 
 			{
 				char* e;
-				// need to strip leading and trailing whitespace so token-pasting works
-
+				// if this is a single line argument, need to strip trailing whitespace so token-pasting works
+				// otherwise we keep the whitespace as-is by re-inserting the leading ws, and so assume multiline macro arguments are not token pasted
 				arrput(copy, 0);
 				p = copy;
 				e = copy + arrlen(copy) - 1;  // get address of NUL
-
-				if (e - 1 > p)
+				
+				if (arg_newlines == 0)
 				{
-					e = (char*)preprocessor_skip_whitespace_reverse(e);
-					assert(e > p);
+					if (e - 1 > p)
+					{
+						e = (char*)preprocessor_skip_whitespace_reverse(e);
+						assert(e > p);
+					}
+					*e = 0;
+				}
+				else
+				{
+					arrinsn(copy, 0, leading_ws_size);
+					memcpy(copy, leading_ws, leading_ws_size);
 				}
-				*e = 0;
 
 #if SSE_READ_PADDING
 				// Need to ensure padding for safe SSE reads without special cases -- normally this comes from a stack allocated buffer with plenty
@@ -3423,7 +3431,7 @@ static void maybe_expand_macro(parse_state* cs, struct macro_definition* pending
 				arrsetlen(copy, e - p);
 				arrput(arguments, copy);
 				cs->src_line_number += arg_newlines;
-
+				cs->dest_line_number += arg_newlines;
 				p = s + 1;	// advance to after terminating character
 			}
 		}
@@ -3507,6 +3515,7 @@ static void maybe_expand_macro(parse_state* cs, struct macro_definition* pending
 
 				arrput(arguments, copy);
 				cs->src_line_number += arg_newlines;
+				cs->dest_line_number += arg_newlines;
 			}
 		}
 
