[shaders] add support in stb_preprocess for #warning directives, and log warnings emitted by preprocessing

diff --git a/Engine/Source/Developer/ShaderPreprocessor/Private/ShaderPreprocessor.cpp b/Engine/Source/Developer/ShaderPreprocessor/Private/ShaderPreprocessor.cpp
--- a/Engine/Source/Developer/ShaderPreprocessor/Private/ShaderPreprocessor.cpp
+++ b/Engine/Source/Developer/ShaderPreprocessor/Private/ShaderPreprocessor.cpp
@@ -810,8 +810,7 @@ bool PreprocessShader(
 			HasError |= (Diagnostic->error_level == PP_RESULT_MODE_error);
 			
 			FString Message = Diagnostic->message;
-			// ignore stb warnings (for now?)
-			if (Diagnostic->error_level == PP_RESULT_MODE_error)
+			if (Diagnostic->error_level == PP_RESULT_MODE_error || Diagnostic->error_level == PP_RESULT_MODE_warning)
 			{
 				FString Filename = Diagnostic->where->filename;
 				Output.LogError(MoveTemp(Filename), MoveTemp(Message), Diagnostic->where->line_number);
diff --git a/Engine/Source/Developer/ShaderPreprocessor/Private/stb_preprocess/preprocessor.c b/Engine/Source/Developer/ShaderPreprocessor/Private/stb_preprocess/preprocessor.c
--- a/Engine/Source/Developer/ShaderPreprocessor/Private/stb_preprocess/preprocessor.c
+++ b/Engine/Source/Developer/ShaderPreprocessor/Private/stb_preprocess/preprocessor.c
@@ -484,6 +484,7 @@ static void output_line_directive(parse_state* cs)
 static uint8 pp_result_mode[PP_RESULT_count] = {
 	PP_RESULT_MODE_no_warning, PP_RESULT_MODE_supplementary,
 	PP_RESULT_MODE_no_warning,	// #undef of undefined symbol is silently ignored per spec
+	PP_RESULT_MODE_warning,
 								// initialize all others to PP_RESULT_MODE_stop
 };
 
@@ -2197,6 +2198,7 @@ enum
 	HASH_undef = 18,
 	HASH_pragma = 24,
 	HASH_ifdef = 30,
+	HASH_warning = 22,
 
 	HASH_none = 31,
 };
@@ -3902,7 +3904,6 @@ static void process_directive(parse_state* cs, conditional_state* cons)
 		case 19:
 		case 20:
 		case 21:
-		case 22:
 		case 23:
 		case 25:
 		case 26:
@@ -3931,7 +3932,7 @@ static void process_directive(parse_state* cs, conditional_state* cons)
 						if (pp_result_mode[PP_RESULT_unrecognized_directive] != PP_RESULT_MODE_no_warning)
 						{
 							static char* directives[] = {
-								"define", "else", "elif", "endif", "error", "if", "ifdef", "ifndef", "include", "line", "pragma", "undef",
+								"define", "else", "elif", "endif", "error", "if", "ifdef", "ifndef", "include", "line", "pragma", "undef", "warning",
 							};
 							char identifier[64] = { 0 }, *match = 0;
 							int i;
@@ -4448,18 +4449,18 @@ static void process_directive(parse_state* cs, conditional_state* cons)
 		}
 
 		case HASH_error:
+		case HASH_warning:
 			p = preprocessor_skip_whitespace_simple(p);
-			const char* err_end = p;
-			while (*err_end != 0 && !(char_is_end_of_line(*err_end)))
-				err_end++;
+			const char* str_end = p;
+			while (*str_end != 0 && !(char_is_end_of_line(*str_end)))
+				str_end++;
 		
-			ptrdiff_t len = (err_end - p) + 1;
-			char* err = (char*)STB_COMMON_MALLOC(len);
-			memcpy(err, p, len - 1);
-			err[len - 1] = 0;
+			ptrdiff_t len = (str_end - p) + 1;
+			char* str = (char*)STB_COMMON_MALLOC(len);
+			memcpy(str, p, len - 1);
+			str[len - 1] = 0;
 
-			error_explicit(cs, PP_RESULT_ERROR, 0, err);
-			break;
+			error_explicit(cs, ((n & 31) == HASH_error) ? PP_RESULT_ERROR : PP_RESULT_explicit_warning, 0, str);
 	}
 
 	// each case above individually scans to end of line so they can
@@ -5179,6 +5180,7 @@ void init_preprocessor(
 	init_directive("line", HASH_line);
 	init_directive("pragma", HASH_pragma);
 	init_directive("undef", HASH_undef);
+	init_directive("warning", HASH_warning);
 
 	loadfile_callback = load_callback;
 	freefile_callback = free_callback;
diff --git a/Engine/Source/Developer/ShaderPreprocessor/Private/stb_preprocess/preprocessor.h b/Engine/Source/Developer/ShaderPreprocessor/Private/stb_preprocess/preprocessor.h
--- a/Engine/Source/Developer/ShaderPreprocessor/Private/stb_preprocess/preprocessor.h
+++ b/Engine/Source/Developer/ShaderPreprocessor/Private/stb_preprocess/preprocessor.h
@@ -8,10 +8,11 @@ enum
 	PP_RESULT_ok,
 	PP_RESULT_supplementary,			 // extra information associated with previous warning/error
 	PP_RESULT_undef_of_undefined_macro,	 // by default this is not an error
-
+	PP_RESULT_explicit_warning,
+	
 	PP_RESULT_internal_error_out_of_memory,
 
-	PP_RESULT_ERROR = 4,
+	PP_RESULT_ERROR = 5,
 
 	PP_RESULT_counter_overflowed,
 	PP_RESULT_too_many_arguments_to_macro,
diff --git a/Engine/Source/Runtime/RenderCore/Private/ShaderCore.cpp b/Engine/Source/Runtime/RenderCore/Private/ShaderCore.cpp
--- a/Engine/Source/Runtime/RenderCore/Private/ShaderCore.cpp
+++ b/Engine/Source/Runtime/RenderCore/Private/ShaderCore.cpp
@@ -1582,6 +1582,8 @@ public:
 			PreprocessShaderInternal(Compiler, Job);
 		}
 
+		Job.Output.Errors.Append(Job.PreprocessOutput.Errors);
+		
 		// decompress if necessary; this is a no-op if source is not compressed.
 		Job.PreprocessOutput.DecompressCode();
 
