[shaders] silence static analysis warning from PVS studio (it doesn't like sizeof(*a) but IMO this is perfectly valid)

diff --git a/Engine/Source/Developer/ShaderPreprocessor/Private/stb_preprocess/stb_ds.h b/Engine/Source/Developer/ShaderPreprocessor/Private/stb_preprocess/stb_ds.h
--- a/Engine/Source/Developer/ShaderPreprocessor/Private/stb_preprocess/stb_ds.h
+++ b/Engine/Source/Developer/ShaderPreprocessor/Private/stb_preprocess/stb_ds.h
@@ -581,7 +581,7 @@ extern "C"
 #define stbds_arrdelswap(a, i) ((a)[i] = stbds_arrlast(a), stbds_header(a)->length -= 1)
 #define stbds_arrinsn(a, i, n)                               \
 	((a) ? stbds_temp(a) = (i) : 0, stbds_arraddn((a), (n)), \
-	 memmove(&(a)[stbds_temp(a) + (n)], &(a)[stbds_temp(a)], sizeof *(a) * (stbds_header(a)->length - (n)-stbds_temp(a))))
+	 memmove(&(a)[stbds_temp(a) + (n)], &(a)[stbds_temp(a)], sizeof *(a) * (stbds_header(a)->length - (n)-stbds_temp(a))))  //-V568
 #define stbds_arrins(a, i, v) (stbds_arrinsn((a), (i), 1), (a)[i] = (v))
 
 #define stbds_arrmaybegrow(a, n) ((!(a) || stbds_header(a)->length + (n) > stbds_header(a)->capacity) ? (stbds_arrgrow(a, n, 0), 0) : 0)
