Move PVS silencing to correct line
#rnx

diff --git a/Engine/Source/Developer/ShaderPreprocessor/Private/stb_preprocess/preprocessor.c b/Engine/Source/Developer/ShaderPreprocessor/Private/stb_preprocess/preprocessor.c
--- a/Engine/Source/Developer/ShaderPreprocessor/Private/stb_preprocess/preprocessor.c
+++ b/Engine/Source/Developer/ShaderPreprocessor/Private/stb_preprocess/preprocessor.c
@@ -3405,7 +3405,7 @@ static void maybe_expand_macro(parse_state* cs, struct macro_definition* pending
 				}
 				else
 				{
-					arrinsn(copy, 0, leading_ws_size);
+					arrinsn(copy, 0, leading_ws_size); //-V568
 					memcpy(copy, leading_ws, leading_ws_size);
 				}
 
diff --git a/Engine/Source/Developer/ShaderPreprocessor/Private/stb_preprocess/stb_ds.h b/Engine/Source/Developer/ShaderPreprocessor/Private/stb_preprocess/stb_ds.h
--- a/Engine/Source/Developer/ShaderPreprocessor/Private/stb_preprocess/stb_ds.h
+++ b/Engine/Source/Developer/ShaderPreprocessor/Private/stb_preprocess/stb_ds.h
@@ -581,7 +581,7 @@ extern "C"
 #define stbds_arrdelswap(a, i) ((a)[i] = stbds_arrlast(a), stbds_header(a)->length -= 1)
 #define stbds_arrinsn(a, i, n)                               \
 	((a) ? stbds_temp(a) = (i) : 0, stbds_arraddn((a), (n)), \
-	 memmove(&(a)[stbds_temp(a) + (n)], &(a)[stbds_temp(a)], sizeof *(a) * (stbds_header(a)->length - (n)-stbds_temp(a))))  //-V568
+	 memmove(&(a)[stbds_temp(a) + (n)], &(a)[stbds_temp(a)], sizeof *(a) * (stbds_header(a)->length - (n)-stbds_temp(a))))
 #define stbds_arrins(a, i, v) (stbds_arrinsn((a), (i), 1), (a)[i] = (v))
 
 #define stbds_arrmaybegrow(a, n) ((!(a) || stbds_header(a)->length + (n) > stbds_header(a)->capacity) ? (stbds_arrgrow(a, n, 0), 0) : 0)
