Shader Preprocessor:  Fix bug where the pointer to the preprocessed text buffer isn't updated when reallocated while appending string printing code.  This hadn't been hit before, because we reserve a large (1.25 MB) buffer to begin with, so encountering the bug requires a shader that uses string printing (itself rare), and of an extremely precise length (within a few KB of the reserve size, so there isn't slack to absorb the appended code).

#rnx

diff --git a/Engine/Source/Developer/ShaderPreprocessor/Private/ShaderPreprocessor.cpp b/Engine/Source/Developer/ShaderPreprocessor/Private/ShaderPreprocessor.cpp
--- a/Engine/Source/Developer/ShaderPreprocessor/Private/ShaderPreprocessor.cpp
+++ b/Engine/Source/Developer/ShaderPreprocessor/Private/ShaderPreprocessor.cpp
@@ -179,7 +179,7 @@ struct FStbPreprocessContext
 		return HasIncludedHeader(PlatformHeader);
 	}
 
-	void ShaderPrintGenerate(char* PreprocessFile, TArray<FShaderDiagnosticData>* OutDiagnosticDatas);
+	void ShaderPrintGenerate(char*& PreprocessFile, TArray<FShaderDiagnosticData>* OutDiagnosticDatas);
 };
 
 static void StbLoadedIncludeTrimPaddingChecked(FStbLoadedInclude* ContentsCached)
@@ -574,7 +574,7 @@ static void StbCustomMacroEnd(const char* OriginalText, void* RawContext, const
 	}
 }
 
-void FStbPreprocessContext::ShaderPrintGenerate(char* PreprocessedFile, TArray<FShaderDiagnosticData>* OutDiagnosticDatas)
+void FStbPreprocessContext::ShaderPrintGenerate(char*& PreprocessedFile, TArray<FShaderDiagnosticData>* OutDiagnosticDatas)
 {
 	// Check if ShaderPrintCommon.ush was included, to decide whether to add the shader print generated code
 	static FString ShaderPrintHeader("/Engine/Private/ShaderPrintCommon.ush");
@@ -630,7 +630,7 @@ void FStbPreprocessContext::ShaderPrintGenerate(char* PreprocessedFile, TArray<F
 	TArray<ANSICHAR> TextCharsAnsi;
 	CopyStringToAnsiCharArray(*TextChars, TextChars.Len(), TextCharsAnsi);
 
-	preprocessor_file_append(PreprocessedFile, TextCharsAnsi.GetData(), TextCharsAnsi.Num() - 1);
+	PreprocessedFile = preprocessor_file_append(PreprocessedFile, TextCharsAnsi.GetData(), TextCharsAnsi.Num() - 1);
 
 	// 4. Insert assert data into shader compilation output for runtime CPU lookup
 	if (OutDiagnosticDatas && TextAssertCount > 0)
@@ -873,4 +873,4 @@ bool PreprocessShader(
 		ShaderOutput.Errors.Add(MoveTemp(Error));
 	}
 	return bSucceeded;
-}
\ No newline at end of file
+}
diff --git a/Engine/Source/Developer/ShaderPreprocessor/Private/stb_preprocess/preprocessor.c b/Engine/Source/Developer/ShaderPreprocessor/Private/stb_preprocess/preprocessor.c
--- a/Engine/Source/Developer/ShaderPreprocessor/Private/stb_preprocess/preprocessor.c
+++ b/Engine/Source/Developer/ShaderPreprocessor/Private/stb_preprocess/preprocessor.c
@@ -5056,7 +5056,7 @@ int preprocessor_file_capacity(char* text)
 	return text ? arrcap(text) : 0;
 }
 
-void preprocessor_file_append(char* text, const char* appended_text, int appended_text_len)
+char* preprocessor_file_append(char* text, const char* appended_text, int appended_text_len)
 {
 	if (text)
 	{
@@ -5071,6 +5071,8 @@ void preprocessor_file_append(char* text, const char* appended_text, int appende
 		// And add a new null terminator
 		text[text_len + appended_text_len - 1] = 0;
 	}
+
+	return text;
 }
 
 void preprocessor_file_free(char* text, pp_diagnostic* pd)
diff --git a/Engine/Source/Developer/ShaderPreprocessor/Private/stb_preprocess/preprocessor.h b/Engine/Source/Developer/ShaderPreprocessor/Private/stb_preprocess/preprocessor.h
--- a/Engine/Source/Developer/ShaderPreprocessor/Private/stb_preprocess/preprocessor.h
+++ b/Engine/Source/Developer/ShaderPreprocessor/Private/stb_preprocess/preprocessor.h
@@ -114,7 +114,7 @@ STB_PP_DEF int preprocessor_file_capacity(char* text);
 
 // Append text to the end of a file generated with preprocess_file.  Can save reallocation overhead relative
 // to appending the text later, as there will usually be slack space available.
-STB_PP_DEF void preprocessor_file_append(char* text, const char* appended_text, int appended_text_len);
+STB_PP_DEF char* preprocessor_file_append(char* text, const char* appended_text, int appended_text_len);
 
 // frees memory allocated by preprocess_file (preprocessed results and diagnostic messages)
 STB_PP_DEF void preprocessor_file_free(char* text, pp_diagnostic* pd);
