[shaders] fix preprocessor infinite loop caused by '[at]' being considered part of the identifier character class, but not being handled properly by copy_and_filter_macro (it appears the RADC compiler from which this code originated allowed [at] in identifiers, but our UE SIMD optimizations missed this quirk. since it's not a valid identifier character in HLSL I opted to just remove it from the character class.

diff --git a/Engine/Source/Developer/ShaderPreprocessor/Private/stb_preprocess/preprocessor.c b/Engine/Source/Developer/ShaderPreprocessor/Private/stb_preprocess/preprocessor.c
--- a/Engine/Source/Developer/ShaderPreprocessor/Private/stb_preprocess/preprocessor.c
+++ b/Engine/Source/Developer/ShaderPreprocessor/Private/stb_preprocess/preprocessor.c
@@ -4747,7 +4747,6 @@ void init_preprocessor_scanner(void)
 	pp_char_class['\''] = PP_CHAR_CLASS_apostrophe;
 	pp_char_class['\\'] = PP_CHAR_CLASS_backslash;
 
-	pp_char_class['@'] = PP_CHAR_CLASS_idtype;	// RADC extension
 	pp_char_class['#'] = PP_CHAR_CLASS_hash;
 
 	// the scanner does the same as above, but doesn't stop on identifiers
