[shaders] bug fixes & improvements to conditional expression parser:
- remove the (broken) handling of large unsigned values (i.e. those > 2^63 and so requiring the sign bit); this doesn't work as written due to various bugs, and is unnecessary due to the HLSL spec explicitly stating all conditional expression values are interpreted as 32 bit signed values (see https://learn.microsoft.com/en-us/windows/win32/direct3dhlsl/dx-graphics-hlsl-appendix-pre-if)
- add explicit overflow checks for affected operations (initial assignment, add, and multiply), using 64 bit signed values for intermediate results
- improve error handling with enumerated values for different error conditions, handle these error cases in the preprocessor with more detail and include source file/line info in the error message
- fix parsing of explicit unsigned/long suffices on conditional expressions; these don't really matter due to the above but this wasn't previously handling uppercase U/L and so throwing spurious errors in these cases
- move tests into an explicit test file and change to the lowlevel test framework, add tests for overflow cases
- add ShaderPreprocessor module to foundation tests suite (limited to desktop platforms) so they are picked up in VS test explorer and lowlevel test CIS
- fix a bug in the lowleveltests initialization causing checks to fire on mac
- add a simple unit test that runs the full preprocessor to trigger an overflow error (setting up some simple boilerplate for implementing other preprocessor tests going forward)

diff --git a/Engine/Source/Developer/LowLevelTestsRunner/Private/TestCommon/CoreUtilities.cpp b/Engine/Source/Developer/LowLevelTestsRunner/Private/TestCommon/CoreUtilities.cpp
--- a/Engine/Source/Developer/LowLevelTestsRunner/Private/TestCommon/CoreUtilities.cpp
+++ b/Engine/Source/Developer/LowLevelTestsRunner/Private/TestCommon/CoreUtilities.cpp
@@ -89,6 +89,7 @@ void CleanupAllThreadPools()
 
 void InitAllThreadPools(bool MultiThreaded)
 {
+	FTaskTagScope Scope(ETaskTag::EGameThread);
 	InitIOThreadPool(MultiThreaded);
 	InitThreadPool(MultiThreaded);
 	InitBackgroundPriorityThreadPool(MultiThreaded);
diff --git a/Engine/Source/Developer/ShaderPreprocessor/Private/stb_preprocess/cond_expr.c b/Engine/Source/Developer/ShaderPreprocessor/Private/stb_preprocess/cond_expr.c
--- a/Engine/Source/Developer/ShaderPreprocessor/Private/stb_preprocess/cond_expr.c
+++ b/Engine/Source/Developer/ShaderPreprocessor/Private/stb_preprocess/cond_expr.c
@@ -8,17 +8,17 @@
 // into 0 or 1.
 
 #include <ctype.h>
+#include <errno.h>
+#include <limits.h>
 #include <stddef.h>
 #include <stdlib.h>
 #include <string.h>
+#include "cond_expr.h"
 
 #ifdef _MSC_VER
 #define intmax_t __int64
-#define uintmax_t unsigned __int64
-#define strtoull (unsigned)strtoul
 #else
 #define intmax_t long long
-#define uintmax_t unsigned long long
 #endif
 
 #if !defined(FORCE_INLINE)
@@ -29,16 +29,6 @@
 #endif
 #endif
 
-typedef struct
-{
-	union
-	{
-		intmax_t iv;
-		uintmax_t uv;
-	};
-	int is_signed;
-} Value;
-
 enum
 {
 	PPLEX_eof = 0,
@@ -57,8 +47,8 @@ typedef struct
 {
 	char* p;
 	int token;
-	Value token_value;
-	int syntax_error;
+	int token_value;
+	int result;
 } ppcexp;  // preprocessor constant expression lexer
 
 // Faster inline version of CRT isspace function
@@ -86,7 +76,7 @@ static FORCE_INLINE unsigned isspace_inline(unsigned p)
 
 // Assuming we found a starting digit, this function checks the next character to see if it ends the number.  If so, it's a single
 // digit number (which is very common in the preprocessor -- most values are zero or one), and we can do a trivial digit character
-// to number conversion, rather than calling the expensive strtoull.  This doesn't try to catch all cases, but it catches most
+// to number conversion, rather than calling the expensive strtoll. This doesn't try to catch all cases, but it catches most
 // cases likely to exist in valid (non syntax error) code.
 static FORCE_INLINE int is_end_of_number(int p)
 {
@@ -277,44 +267,53 @@ static void ce_next(ppcexp* c)
 				++p;
 			while (is_token_character(*p));
 			// any Identifier is a macro that WASN'T defined, so it evaluates as 0
-			c->token_value.iv = 0;
-			c->token_value.is_signed = 1;
+			c->token_value = 0;
 			c->token = PPLEX_int_literal;
 			break;
 		}
 
+#define parse_check(c, start, offset, base)						\
+	{															\
+		char* end;												\
+		intmax_t parsed = strtoll(start + offset, &end, base);	\
+		if (parsed > INT_MAX || parsed < INT_MIN)				\
+		{														\
+			c->result = CE_RESULT_overflow;						\
+		}														\
+		else if (offset && (end == start + offset))				\
+		{														\
+			c->result = CE_RESULT_syntax_error;					\
+		}														\
+		else													\
+		{														\
+			c->token_value = (int)parsed;						\
+		}														\
+		start = end;											\
+	}
+
 		case '0':
 			if (is_end_of_number(p[1]))
 			{
-				c->token_value.iv = 0;
-				c->token_value.is_signed = 1;
+				c->token_value = 0;
 				c->token = PPLEX_int_literal;
 				p++;
 				break;
 			}
 			if (p[1] == 'x' || p[1] == 'X')
 			{
-				char* q;
-				c->token_value.uv = strtoull(p + 2, &q, 16);
-				if (q == p + 2)
-					c->syntax_error = 1;
-				p = q;
+				parse_check(c, p, 2, 16);
 			}
 			else if (p[1] == 'b' || p[1] == 'B')
 			{
-				char* q;
-				c->token_value.uv = strtoull(p + 2, &q, 2);
-				if (q == p + 2)
-					c->syntax_error = 1;
-				p = q;
+				parse_check(c, p, 2, 2);
 			}
 			else if (p[1] == '.')
 			{
-				c->syntax_error = 1;
+				c->result = CE_RESULT_syntax_error;
 			}
 			else
 			{
-				c->token_value.uv = strtoull(p, &p, 8);
+				parse_check(c, p, 0, 8);
 			}
 			goto numeric_suffixes;
 
@@ -329,19 +328,16 @@ static void ce_next(ppcexp* c)
 		case '9':
 			if (is_end_of_number(p[1]))
 			{
-				c->token_value.iv = *p - '0';
-				c->token_value.is_signed = 1;
+				c->token_value = *p - '0';
 				c->token = PPLEX_int_literal;
 				p++;
 				break;
 			}
-			c->token_value.uv = strtoull(p, &p, 10);
+			parse_check(c, p, 0, 10);
+
 		numeric_suffixes:
-			c->token_value.is_signed = (c->token_value.uv < ((uintmax_t)1 << 63));
-			while (*p == 'u' || *p == 'l')
+			while (*p == 'u' || *p == 'U' || *p == 'l' || *p == 'L')
 			{
-				if (*p == 'u')
-					c->token_value.is_signed = 1;
 				++p;
 			}
 			c->token = PPLEX_int_literal;
@@ -350,20 +346,20 @@ static void ce_next(ppcexp* c)
 	c->p = p;
 }
 
-static Value ce_primary(ppcexp* c)
+static int ce_primary(ppcexp* c)
 {
-	Value v;
+	int v;
 	if (c->token != PPLEX_int_literal)
-		c->syntax_error = 1;
+		c->result = CE_RESULT_syntax_error;
 	v = c->token_value;
 	ce_next(c);
 	return v;
 }
 
-static Value ce_cond(ppcexp* c);
-static Value ce_unary(ppcexp* c)
+static int ce_cond(ppcexp* c);
+static int ce_unary(ppcexp* c)
 {
-	Value v;
+	int v;
 	switch (c->token)
 	{
 		case '+':
@@ -372,24 +368,23 @@ static Value ce_unary(ppcexp* c)
 		case '-':
 			ce_next(c);
 			v = ce_unary(c);
-			v.iv = -v.iv;
+			v = -v;
 			return v;
 		case '!':
 			ce_next(c);
 			v = ce_unary(c);
-			v.iv = (v.uv == 0);
-			v.is_signed = 1;
+			v = (v == 0);
 			return v;
 		case '~':
 			ce_next(c);
 			v = ce_unary(c);
-			v.uv = ~v.uv;
+			v = ~v;
 			return v;
 		case '(':
 			ce_next(c);
 			v = ce_cond(c);
 			if (c->token != ')')
-				c->syntax_error = 1;
+				c->result = CE_RESULT_syntax_error;
 			ce_next(c);
 			return v;
 		default:
@@ -397,292 +392,257 @@ static Value ce_unary(ppcexp* c)
 	}
 }
 
-static Value ce_mul(ppcexp* c)
+static int ce_mul(ppcexp* c)
 {
-	Value left = ce_unary(c);
+	int left = ce_unary(c);
 	while (c->token == '*' || c->token == '/' || c->token == '%')
 	{
 		int token = c->token;
-		Value right;
+		int right;
 		ce_next(c);
 		right = ce_unary(c);
 		if (token == '*')
-			left.uv *= right.uv;
+		{
+			intmax_t result_wide = (intmax_t)left * (intmax_t)right;
+			if (result_wide > INT_MAX || result_wide < INT_MIN)
+			{
+				c->result = CE_RESULT_overflow;
+			}
+			left = (int)result_wide;
+		}
 		else
 		{
-			if (right.uv == 0)
+			if (right == 0)
 			{
-				if (c->syntax_error == 0)
-					c->syntax_error = -1;
+				if (c->result == CE_RESULT_ok)
+				{
+					c->result = CE_RESULT_division_by_zero;
+				}
 			}
-			else if (left.is_signed & right.is_signed)
-				if (token == '/')
-					left.iv /= right.iv;
-				else
-					left.iv %= right.iv;
 			else if (token == '/')
-				left.uv /= right.uv;
+			{
+				left /= right;
+			}
 			else
-				left.uv %= right.uv;
-			left.is_signed &= right.is_signed;
+			{
+				left %= right;
+			}
 		}
 	}
 	return left;
 }
 
-static Value ce_sum(ppcexp* c)
+static int ce_sum(ppcexp* c)
 {
-	Value left = ce_mul(c);
+	int left = ce_mul(c);
 	while (c->token == '+' || c->token == '-')
 	{
 		int token = c->token;
-		Value right;
+		int right;
+		intmax_t result_wide;
 		ce_next(c);
 		right = ce_mul(c);
 		if (token == '+')
-			left.uv += right.uv;
+		{
+			result_wide = (intmax_t)left + (intmax_t)right;
+		}
 		else
-			left.uv -= right.uv;
-		left.is_signed &= right.is_signed;
+		{
+			result_wide = (intmax_t)left - (intmax_t)right;
+		}
+		if (result_wide > INT_MAX || result_wide < INT_MIN)
+		{
+			c->result = CE_RESULT_overflow;
+		}
+		left = (int)result_wide;
 	}
 	return left;
 }
 
-static Value ce_shift(ppcexp* c)
+static int ce_shift(ppcexp* c)
 {
-	Value left = ce_sum(c);
+	int left = ce_sum(c);
 	while (c->token == PPLEX_shift_left || c->token == PPLEX_shift_right)
 	{
 		int tok = c->token;
-		Value right;
+		int right;
 		ce_next(c);
 		right = ce_sum(c);
 		if (tok == PPLEX_shift_right)
-			if (left.is_signed)
-				left.iv = (left.iv >> right.uv);
-			else
-				left.uv = (left.uv >> right.uv);
+		{
+			left = (left >> right);
+		}
 		else
-			left.uv = (left.uv << right.uv);
+		{
+			left = (left << right);
+		}
 	}
 	return left;
 }
 
-static Value ce_compare_inequality(ppcexp* c)
+static int ce_compare_inequality(ppcexp* c)
 {
-	Value left = ce_shift(c);
+	int left = ce_shift(c);
 	for (;;)
 	{
-		Value right;
+		int right;
 		switch (c->token)
 		{
 			case PPLEX_less_eq:
 				ce_next(c);
 				right = ce_shift(c);
-				left.iv = (left.is_signed & right.is_signed) ? (left.iv <= right.iv) : (left.uv <= right.uv);
+				left = (left <= right);
 				break;
 			case PPLEX_greater_eq:
 				ce_next(c);
 				right = ce_shift(c);
-				left.iv = (left.is_signed & right.is_signed) ? (left.iv >= right.iv) : (left.uv >= right.uv);
+				left = (left >= right);
 				break;
 			case '<':
 				ce_next(c);
 				right = ce_shift(c);
-				left.iv = (left.is_signed & right.is_signed) ? (left.iv < right.iv) : (left.uv < right.uv);
+				left = (left < right);
 				break;
 			case '>':
 				ce_next(c);
 				right = ce_shift(c);
-				left.iv = (left.is_signed & right.is_signed) ? (left.iv > right.iv) : (left.uv > right.uv);
+				left = (left > right);
 				break;
 			default:
 				return left;
 		}
-		left.is_signed = 1;
 	}
 }
 
-static Value ce_compare_equality(ppcexp* c)
+static int ce_compare_equality(ppcexp* c)
 {
-	Value left = ce_compare_inequality(c), right;
+	int left = ce_compare_inequality(c), right;
 	while (c->token == PPLEX_eq_eq || c->token == PPLEX_bang_eq)
 	{
 		int token = c->token;
 		ce_next(c);
 		right = ce_compare_inequality(c);
-		left.uv = (left.uv == right.uv) ^ (token == PPLEX_bang_eq);
-		left.is_signed = 1;
+		left = (left == right) ^ (token == PPLEX_bang_eq);
 	}
 	return left;
 }
 
-static Value ce_bitwise_and(ppcexp* c)
+static int ce_bitwise_and(ppcexp* c)
 {
-	Value left = ce_compare_equality(c), right;
+	int left = ce_compare_equality(c), right;
 	while (c->token == '&')
 	{
 		ce_next(c);
 		right = ce_compare_equality(c);
-		left.uv &= right.uv;
-		left.is_signed &= right.is_signed;
+		left &= right;
 	}
 	return left;
 }
 
-static Value ce_bitwise_xor(ppcexp* c)
+static int ce_bitwise_xor(ppcexp* c)
 {
-	Value left = ce_bitwise_and(c), right;
+	int left = ce_bitwise_and(c), right;
 	while (c->token == '^')
 	{
 		ce_next(c);
 		right = ce_bitwise_and(c);
-		left.uv ^= right.uv;
-		left.is_signed &= right.is_signed;
+		left ^= right;
 	}
 	return left;
 }
 
-static Value ce_bitwise_or(ppcexp* c)
+static int ce_bitwise_or(ppcexp* c)
 {
-	Value left = ce_bitwise_xor(c), right;
+	int left = ce_bitwise_xor(c), right;
 	while (c->token == '|')
 	{
 		ce_next(c);
 		right = ce_bitwise_xor(c);
-		left.uv |= right.uv;
-		left.is_signed &= right.is_signed;
+		left |= right;
 	}
 	return left;
 }
 
-static Value ce_logical_and(ppcexp* c)
+static int ce_logical_and(ppcexp* c)
 {
-	Value left = ce_bitwise_or(c), right;
+	int left = ce_bitwise_or(c), right;
 	while (c->token == PPLEX_and_and)
 	{
-		int err = c->syntax_error;
+		int result = c->result;
 		ce_next(c);
 		right = ce_bitwise_or(c);
-		if (left.uv == 0 && c->syntax_error < 0 && err == 0)
-			c->syntax_error = 0;  // suppress unevaluated div by 0
-		left.uv = left.uv && right.uv;
-		left.is_signed = 1;
+		if (left == 0 && c->result == CE_RESULT_division_by_zero && result == CE_RESULT_ok)
+			c->result = CE_RESULT_ok;  // suppress unevaluated div by 0
+		left = left && right;
 	}
 	return left;
 }
 
-static Value ce_logical_or(ppcexp* c)
+static int ce_logical_or(ppcexp* c)
 {
-	Value left = ce_logical_and(c), right;
+	int left = ce_logical_and(c), right;
 	while (c->token == PPLEX_or_or)
 	{
-		int err = c->syntax_error;
+		int result = c->result;
 		ce_next(c);
 		right = ce_logical_and(c);
-		if (left.uv != 0 && c->syntax_error < 0 && err == 0)
-			c->syntax_error = 0;  // suppress unevaluated div by 0
-		left.uv = left.uv || right.uv;
-		left.is_signed = 1;
+		if (left != 0 && c->result == CE_RESULT_division_by_zero && result == CE_RESULT_ok)
+			c->result = CE_RESULT_ok;  // suppress unevaluated div by 0
+		left = left || right;
 	}
 	return left;
 }
 
-static Value ce_cond(ppcexp* c)
+static int ce_cond(ppcexp* c)
 {
-	Value cond = ce_logical_or(c);
+	int cond = ce_logical_or(c);
 	if (c->token != '?')
 		return cond;
 	else
 	{
-		int err;
-		Value result, left;
-		err = c->syntax_error;
+		int v, result, left;
+		result = c->result;
 		ce_next(c);
 		left = ce_cond(c);
 		if (c->token != ':')
 		{
-			c->syntax_error = 1;
+			c->result = CE_RESULT_syntax_error;
 			return cond;
 		}
-		if (cond.uv == 0 && c->syntax_error < 0 && err == 0)
-			c->syntax_error = 0;  // suppress unevaluated div by 0
-		err = c->syntax_error;
+		if (cond == 0 && c->result == CE_RESULT_division_by_zero && result == CE_RESULT_ok)
+			c->result = CE_RESULT_ok;  // suppress unevaluated div by 0
+		result = c->result;
 		ce_next(c);
-		result = ce_cond(c);
-		if (cond.uv != 0 && c->syntax_error < 0 && err == 0)
-			c->syntax_error = 0;  // suppress unevaluated div by 0
-		result.uv = cond.uv != 0 ? left.uv : result.uv;
-		result.is_signed &= left.is_signed;
-		return result;
+		v = ce_cond(c);
+		if (cond != 0 && c->result == CE_RESULT_division_by_zero && result == CE_RESULT_ok)
+			c->result = CE_RESULT_ok;  // suppress unevaluated div by 0
+		v = cond != 0 ? left : v;
+		return v;
 	}
 }
 
-int evaluate_integer_constant_expression_as_condition(char* p, int* syntax_error)
+int evaluate_integer_constant_expression_as_condition(char* p, int* result)
 {
-	Value v;
+	int v;
 	ppcexp c = {0};
 	c.p = p;
 	ce_next(&c);
 	v = ce_cond(&c);
 	if (c.token != PPLEX_eof)
-		c.syntax_error = 1;
-	*syntax_error = c.syntax_error;
-	return v.uv != 0;
+		c.result = CE_RESULT_syntax_error;
+	*result = c.result;
+	return v != 0;
 }
 
-#ifdef _MSC_VER
-__int64
-#else
-long long
-#endif
-evaluate_integer_constant_expression(char* p, int* syntax_error)
+int evaluate_integer_constant_expression(char* p, int* result)
 {
-	Value v;
+	int v;
 	ppcexp c = {0};
 	c.p = p;
 	ce_next(&c);
 	v = ce_cond(&c);
 	if (c.token != PPLEX_eof)
-		c.syntax_error = 1;
-	*syntax_error = c.syntax_error;
-	return v.iv;
-}
-
-#include <assert.h>
-void test_integer_constant_expression(void)
-{
-	int err;
-	assert(evaluate_integer_constant_expression("4+3*2", &err) == 10);
-	assert(err == 0);
-	assert(evaluate_integer_constant_expression("4+3*2==10", &err) == 1);
-	assert(err == 0);
-	assert(evaluate_integer_constant_expression("0xffffffff+1==1<<32", &err) == 1);
-	assert(err == 0);
-	evaluate_integer_constant_expression("1+=1", &err);
-	assert(err > 0);
-	evaluate_integer_constant_expression("1+(", &err);
-	assert(err > 0);
-	evaluate_integer_constant_expression("(1+", &err);
-	assert(err > 0);
-	assert(evaluate_integer_constant_expression("2/1", &err) == 2);
-	assert(err == 0);
-	evaluate_integer_constant_expression("1/0", &err);
-	assert(err < 0);
-	evaluate_integer_constant_expression("1 || 1/0", &err);
-	assert(err == 0);
-	evaluate_integer_constant_expression("0 || 1/0", &err);
-	assert(err < 0);
-	evaluate_integer_constant_expression("1 && 1/0", &err);
-	assert(err < 0);
-	evaluate_integer_constant_expression("0 && 1/0", &err);
-	assert(err == 0);
-	evaluate_integer_constant_expression("0 ? 1/0 : 2/1", &err);
-	assert(err == 0);
-	evaluate_integer_constant_expression("1 ? 1/0 : 2/1", &err);
-	assert(err < 0);
-	evaluate_integer_constant_expression("0 ? 2/1 : 1/0", &err);
-	assert(err < 0);
-	evaluate_integer_constant_expression("1 ? 2/1 : 1/0", &err);
-	assert(err == 0);
+		c.result = CE_RESULT_syntax_error;
+	*result = c.result;
+	return v;
 }
diff --git a/Engine/Source/Developer/ShaderPreprocessor/Private/stb_preprocess/cond_expr.h b/Engine/Source/Developer/ShaderPreprocessor/Private/stb_preprocess/cond_expr.h
--- a/Engine/Source/Developer/ShaderPreprocessor/Private/stb_preprocess/cond_expr.h
+++ b/Engine/Source/Developer/ShaderPreprocessor/Private/stb_preprocess/cond_expr.h
@@ -2,11 +2,19 @@
 
 #pragma once
 
-#ifdef _MSC_VER
-__int64
+enum
+{
+	CE_RESULT_ok,
+	CE_RESULT_syntax_error,
+	CE_RESULT_division_by_zero,
+	CE_RESULT_overflow,
+};
+
+#ifdef __cplusplus
+#define CONDEXPR_DEF extern "C"
 #else
-long long
+#define CONDEXPR_DEF extern
 #endif
-evaluate_integer_constant_expression(char* p, int* syntax_error);
-extern int evaluate_integer_constant_expression_as_condition(char* p, int* syntax_error);
-extern void test_integer_constant_expression(void);
+
+CONDEXPR_DEF int evaluate_integer_constant_expression(char* p, int* error);
+CONDEXPR_DEF int evaluate_integer_constant_expression_as_condition(char* p, int* error);
\ No newline at end of file
diff --git a/Engine/Source/Developer/ShaderPreprocessor/Private/stb_preprocess/preprocessor.c b/Engine/Source/Developer/ShaderPreprocessor/Private/stb_preprocess/preprocessor.c
--- a/Engine/Source/Developer/ShaderPreprocessor/Private/stb_preprocess/preprocessor.c
+++ b/Engine/Source/Developer/ShaderPreprocessor/Private/stb_preprocess/preprocessor.c
@@ -4328,20 +4328,36 @@ static void process_directive(parse_state* cs, conditional_state* cons)
 		case HASH_if:
 		{
 			ifdef_info ifinfo;
-			int result, syntax_error;
+			int result, result_code;
 
 			ifinfo.line = cs->src_line_number;
 			ifinfo.type = HASH_if;
 
 			p = preprocessor_skip_whitespace_simple(p);
-			result = evaluate_if(cs, p, &syntax_error);
+			result = evaluate_if(cs, p, &result_code);
 
-			if (syntax_error)
+			if (result_code != CE_RESULT_ok)
 			{
-				if (do_error(cs, "Syntax error in conditional expression of '#if'"))
-					return;
-				else
-					result = 0;
+				char* err_msg;
+				int ce_result;
+				switch (result_code)
+				{
+				case CE_RESULT_division_by_zero:
+					err_msg = "Division by zero occurred when evaluating conditional expression";
+					ce_result = PP_RESULT_conditional_expression_division_by_zero;
+					break;
+				case CE_RESULT_overflow:
+					err_msg = "Overflow occurred when evaluating conditional expression";
+					ce_result = PP_RESULT_conditional_expression_overflow;
+					break;
+				case CE_RESULT_syntax_error:
+					err_msg = "Syntax error occurred when evaluating conditional expression";
+					ce_result = PP_RESULT_syntax_error_conditional_expression;
+					break;
+				}
+				
+				do_error_explicit(cs, ce_result, cs->src_line_number, err_msg);
+				return;
 			}
 
 			arrpush(c->ifdef_stack, ifinfo);
@@ -5212,7 +5228,6 @@ void init_preprocessor(
 
 	init_preprocessor_scanner();
 	init_char_type();
-	// test_integer_constant_expression();
 
 	init_directive("define", HASH_define);
 	init_directive("else", HASH_else);
diff --git a/Engine/Source/Developer/ShaderPreprocessor/Private/stb_preprocess/preprocessor.h b/Engine/Source/Developer/ShaderPreprocessor/Private/stb_preprocess/preprocessor.h
--- a/Engine/Source/Developer/ShaderPreprocessor/Private/stb_preprocess/preprocessor.h
+++ b/Engine/Source/Developer/ShaderPreprocessor/Private/stb_preprocess/preprocessor.h
@@ -26,6 +26,9 @@ enum
 	PP_RESULT_malformed_pragma_message,
 	PP_RESULT_malformed_line_filename,
 	PP_RESULT_invalid_line_linenumber,
+	PP_RESULT_syntax_error_conditional_expression,
+	PP_RESULT_conditional_expression_overflow,
+	PP_RESULT_conditional_expression_division_by_zero,
 
 	PP_RESULT_count
 };
diff --git a/Engine/Source/Developer/ShaderPreprocessor/Tests/ShaderPreprocessorTests.cpp b/Engine/Source/Developer/ShaderPreprocessor/Tests/ShaderPreprocessorTests.cpp
new file mode 100644
--- /dev/null
+++ b/Engine/Source/Developer/ShaderPreprocessor/Tests/ShaderPreprocessorTests.cpp
@@ -0,0 +1,110 @@
+// Copyright Epic Games, Inc. All Rights Reserved.
+
+#if WITH_TESTS
+
+#include "stb_preprocess/cond_expr.h"
+#include "stb_preprocess/preprocessor.h"
+#include "ShaderSource.h"
+#include "Tests/TestHarnessAdapter.h"
+
+int TestEvaluate(const char* Expression, int& Result)
+{
+	return evaluate_integer_constant_expression(const_cast<char*>(Expression), &Result);
+}
+
+TEST_CASE_NAMED(FConditionalExpressionTests, "Shaders::ShaderPreprocessor::ConditionalExpression", "[EditorContext][EngineFilter]")
+{
+	int Result;
+	CHECK(TestEvaluate("4+3*2", Result) == 10);
+	CHECK(Result == CE_RESULT_ok);
+	CHECK(TestEvaluate("4+3*2==10", Result) == 1);
+	CHECK(Result == CE_RESULT_ok);
+	CHECK(TestEvaluate("4-3*2", Result) == -2);
+	CHECK(Result == CE_RESULT_ok);
+	CHECK(TestEvaluate("-4/2u", Result) == -2);
+	CHECK(Result == CE_RESULT_ok);
+	CHECK(TestEvaluate("4u-5", Result) == -1);
+	CHECK(Result == CE_RESULT_ok);
+	CHECK(TestEvaluate("4U-5", Result) == -1);
+	CHECK(Result == CE_RESULT_ok);
+	TestEvaluate("10000000000000000000000", Result);
+	CHECK(Result == CE_RESULT_overflow);
+	TestEvaluate("0xffffffff+1", Result);
+	CHECK(Result == CE_RESULT_overflow);
+	TestEvaluate("0xffffffff*2", Result);
+	CHECK(Result == CE_RESULT_overflow);
+	TestEvaluate("-2147483647 - 2", Result);
+	CHECK(Result == CE_RESULT_overflow);
+	TestEvaluate("1+=1", Result);
+	CHECK(Result == CE_RESULT_syntax_error);
+	TestEvaluate("1+(", Result);
+	CHECK(Result == CE_RESULT_syntax_error);
+	TestEvaluate("(1+", Result);
+	CHECK(Result == CE_RESULT_syntax_error);
+	CHECK(TestEvaluate("2/1", Result) == 2);
+	CHECK(Result == CE_RESULT_ok);
+	TestEvaluate("1/0", Result);
+	CHECK(Result == CE_RESULT_division_by_zero);
+	TestEvaluate("1 || 1/0", Result);
+	CHECK(Result == CE_RESULT_ok);
+	TestEvaluate("0 || 1/0", Result);
+	CHECK(Result == CE_RESULT_division_by_zero);
+	TestEvaluate("1 && 1/0", Result);
+	CHECK(Result == CE_RESULT_division_by_zero);
+	TestEvaluate("0 && 1/0", Result);
+	CHECK(Result == CE_RESULT_ok);
+	TestEvaluate("0 ? 1/0 : 2/1", Result);
+	CHECK(Result == CE_RESULT_ok);
+	TestEvaluate("1 ? 1/0 : 2/1", Result);
+	CHECK(Result == CE_RESULT_division_by_zero);
+	TestEvaluate("0 ? 2/1 : 1/0", Result);
+	CHECK(Result == CE_RESULT_division_by_zero);
+	TestEvaluate("1 ? 2/1 : 1/0", Result);
+	CHECK(Result == CE_RESULT_ok);
+}
+
+
+const ANSICHAR* TestLoadFile(const ANSICHAR* Filename, void* RawContext, size_t* OutLength)
+{
+	const FShaderSource::FViewType* View = reinterpret_cast<const FShaderSource::FViewType*>(RawContext);
+	*OutLength = View->Len();
+	return View->GetData();
+}
+
+void TestFreeFile(const ANSICHAR* Filename, const ANSICHAR* Contents, void* RawContext)
+{
+	// noop, no files are loaded for these tests but preprocessor calls the freefile callback unconditionally
+}
+
+struct FPreprocessTestResult
+{
+	char* Source = nullptr;
+	int NumDiagnostics = 0;
+	pp_diagnostic* Diagnostics = nullptr;
+	~FPreprocessTestResult()
+	{
+		preprocessor_file_free(Source, Diagnostics);
+	}
+};
+
+FPreprocessTestResult ExecutePreprocessTest(const FShaderSource::FViewType& Source)
+{
+	FPreprocessTestResult Result;
+	init_preprocessor(&TestLoadFile, &TestFreeFile, nullptr, nullptr, nullptr);
+	Result.Source = preprocess_file("preprocessor_test", const_cast<FShaderSource::FViewType*>(&Source), nullptr, 0u, &Result.Diagnostics, &Result.NumDiagnostics, nullptr, 0u);
+	return Result;
+}
+
+TEST_CASE_NAMED(FConditionalExpressionErrorReporting, "Shaders::ShaderPreprocessor::ConditionalExpressionErrorReporting", "[EditorContext][EngineFilter]")
+{
+	FShaderSource::FViewType Source = SHADER_SOURCE_VIEWLITERAL(
+R"(#if 0xFFFFFFFF+1 == 2
+#error "Unreachable"
+#endif)");
+	FPreprocessTestResult Result = ExecutePreprocessTest(Source);
+	CHECK(Result.NumDiagnostics == 1);
+	CHECK(Result.Diagnostics[0].where->line_number == 1);
+	CHECK(strstr(Result.Diagnostics[0].message, "Overflow") != nullptr);
+}
+
+#endif
diff --git a/Engine/Source/Programs/LowLevelTests/FoundationTests/FoundationTests.Build.cs b/Engine/Source/Programs/LowLevelTests/FoundationTests/FoundationTests.Build.cs
--- a/Engine/Source/Programs/LowLevelTests/FoundationTests/FoundationTests.Build.cs
+++ b/Engine/Source/Programs/LowLevelTests/FoundationTests/FoundationTests.Build.cs
@@ -55,7 +55,8 @@ public class FoundationTests : TestModuleRules
 		{
 			PrivateDependencyModuleNames.AddRange(
 				new string[] {
-					"DesktopPlatform"
+					"DesktopPlatform",
+					"ShaderPreprocessor"
 				});
 		}
 	}
