[Backout] - CL36859730
Original CL Desc
-----------------------------------------------------------------
[shaders] bug fixes & improvements to conditional expression parser:
- remove the (broken) handling of large unsigned values (i.e. those > 2^63 and so requiring the sign bit); this doesn't work as written due to various bugs, and is unnecessary due to the HLSL spec explicitly stating all conditional expression values are interpreted as 32 bit signed values (see https://learn.microsoft.com/en-us/windows/win32/direct3dhlsl/dx-graphics-hlsl-appendix-pre-if)
- add explicit overflow checks for affected operations (initial assignment, add, and multiply), using 64 bit signed values for intermediate results
- improve error handling with enumerated values for different error conditions, handle these error cases in the preprocessor with more detail and include source file/line info in the error message
- fix parsing of explicit unsigned/long suffices on conditional expressions; these don't really matter due to the above but this wasn't previously handling uppercase U/L and so throwing spurious errors in these cases
- move tests into an explicit test file and change to the lowlevel test framework, add tests for overflow cases
- add ShaderPreprocessor module to foundation tests suite (limited to desktop platforms) so they are picked up in VS test explorer and lowlevel test CIS
- fix a bug in the lowleveltests initialization causing checks to fire on mac
- add a simple unit test that runs the full preprocessor to trigger an overflow error (setting up some simple boilerplate for implementing other preprocessor tests going forward)

diff --git a/Engine/Source/Developer/LowLevelTestsRunner/Private/TestCommon/CoreUtilities.cpp b/Engine/Source/Developer/LowLevelTestsRunner/Private/TestCommon/CoreUtilities.cpp
--- a/Engine/Source/Developer/LowLevelTestsRunner/Private/TestCommon/CoreUtilities.cpp
+++ b/Engine/Source/Developer/LowLevelTestsRunner/Private/TestCommon/CoreUtilities.cpp
@@ -89,7 +89,6 @@ void CleanupAllThreadPools()
 
 void InitAllThreadPools(bool MultiThreaded)
 {
-	FTaskTagScope Scope(ETaskTag::EGameThread);
 	InitIOThreadPool(MultiThreaded);
 	InitThreadPool(MultiThreaded);
 	InitBackgroundPriorityThreadPool(MultiThreaded);
diff --git a/Engine/Source/Developer/ShaderPreprocessor/Private/stb_preprocess/cond_expr.c b/Engine/Source/Developer/ShaderPreprocessor/Private/stb_preprocess/cond_expr.c
--- a/Engine/Source/Developer/ShaderPreprocessor/Private/stb_preprocess/cond_expr.c
+++ b/Engine/Source/Developer/ShaderPreprocessor/Private/stb_preprocess/cond_expr.c
@@ -8,17 +8,17 @@
 // into 0 or 1.
 
 #include <ctype.h>
-#include <errno.h>
-#include <limits.h>
 #include <stddef.h>
 #include <stdlib.h>
 #include <string.h>
-#include "cond_expr.h"
 
 #ifdef _MSC_VER
 #define intmax_t __int64
+#define uintmax_t unsigned __int64
+#define strtoull (unsigned)strtoul
 #else
 #define intmax_t long long
+#define uintmax_t unsigned long long
 #endif
 
 #if !defined(FORCE_INLINE)
@@ -29,6 +29,16 @@
 #endif
 #endif
 
+typedef struct
+{
+	union
+	{
+		intmax_t iv;
+		uintmax_t uv;
+	};
+	int is_signed;
+} Value;
+
 enum
 {
 	PPLEX_eof = 0,
@@ -47,8 +57,8 @@ typedef struct
 {
 	char* p;
 	int token;
-	int token_value;
-	int result;
+	Value token_value;
+	int syntax_error;
 } ppcexp;  // preprocessor constant expression lexer
 
 // Faster inline version of CRT isspace function
@@ -76,7 +86,7 @@ static FORCE_INLINE unsigned isspace_inline(unsigned p)
 
 // Assuming we found a starting digit, this function checks the next character to see if it ends the number.  If so, it's a single
 // digit number (which is very common in the preprocessor -- most values are zero or one), and we can do a trivial digit character
-// to number conversion, rather than calling the expensive strtoll. This doesn't try to catch all cases, but it catches most
+// to number conversion, rather than calling the expensive strtoull.  This doesn't try to catch all cases, but it catches most
 // cases likely to exist in valid (non syntax error) code.
 static FORCE_INLINE int is_end_of_number(int p)
 {
@@ -267,53 +277,44 @@ static void ce_next(ppcexp* c)
 				++p;
 			while (is_token_character(*p));
 			// any Identifier is a macro that WASN'T defined, so it evaluates as 0
-			c->token_value = 0;
+			c->token_value.iv = 0;
+			c->token_value.is_signed = 1;
 			c->token = PPLEX_int_literal;
 			break;
 		}
 
-#define parse_check(c, start, offset, base)						\
-	{															\
-		char* end;												\
-		intmax_t parsed = strtoll(start + offset, &end, base);	\
-		if (parsed > INT_MAX || parsed < INT_MIN)				\
-		{														\
-			c->result = CE_RESULT_overflow;						\
-		}														\
-		else if (offset && (end == start + offset))				\
-		{														\
-			c->result = CE_RESULT_syntax_error;					\
-		}														\
-		else													\
-		{														\
-			c->token_value = (int)parsed;						\
-		}														\
-		start = end;											\
-	}
-
 		case '0':
 			if (is_end_of_number(p[1]))
 			{
-				c->token_value = 0;
+				c->token_value.iv = 0;
+				c->token_value.is_signed = 1;
 				c->token = PPLEX_int_literal;
 				p++;
 				break;
 			}
 			if (p[1] == 'x' || p[1] == 'X')
 			{
-				parse_check(c, p, 2, 16);
+				char* q;
+				c->token_value.uv = strtoull(p + 2, &q, 16);
+				if (q == p + 2)
+					c->syntax_error = 1;
+				p = q;
 			}
 			else if (p[1] == 'b' || p[1] == 'B')
 			{
-				parse_check(c, p, 2, 2);
+				char* q;
+				c->token_value.uv = strtoull(p + 2, &q, 2);
+				if (q == p + 2)
+					c->syntax_error = 1;
+				p = q;
 			}
 			else if (p[1] == '.')
 			{
-				c->result = CE_RESULT_syntax_error;
+				c->syntax_error = 1;
 			}
 			else
 			{
-				parse_check(c, p, 0, 8);
+				c->token_value.uv = strtoull(p, &p, 8);
 			}
 			goto numeric_suffixes;
 
@@ -328,16 +329,19 @@ static void ce_next(ppcexp* c)
 		case '9':
 			if (is_end_of_number(p[1]))
 			{
-				c->token_value = *p - '0';
+				c->token_value.iv = *p - '0';
+				c->token_value.is_signed = 1;
 				c->token = PPLEX_int_literal;
 				p++;
 				break;
 			}
-			parse_check(c, p, 0, 10);
-
+			c->token_value.uv = strtoull(p, &p, 10);
 		numeric_suffixes:
-			while (*p == 'u' || *p == 'U' || *p == 'l' || *p == 'L')
+			c->token_value.is_signed = (c->token_value.uv < ((uintmax_t)1 << 63));
+			while (*p == 'u' || *p == 'l')
 			{
+				if (*p == 'u')
+					c->token_value.is_signed = 1;
 				++p;
 			}
 			c->token = PPLEX_int_literal;
@@ -346,20 +350,20 @@ static void ce_next(ppcexp* c)
 	c->p = p;
 }
 
-static int ce_primary(ppcexp* c)
+static Value ce_primary(ppcexp* c)
 {
-	int v;
+	Value v;
 	if (c->token != PPLEX_int_literal)
-		c->result = CE_RESULT_syntax_error;
+		c->syntax_error = 1;
 	v = c->token_value;
 	ce_next(c);
 	return v;
 }
 
-static int ce_cond(ppcexp* c);
-static int ce_unary(ppcexp* c)
+static Value ce_cond(ppcexp* c);
+static Value ce_unary(ppcexp* c)
 {
-	int v;
+	Value v;
 	switch (c->token)
 	{
 		case '+':
@@ -368,23 +372,24 @@ static int ce_unary(ppcexp* c)
 		case '-':
 			ce_next(c);
 			v = ce_unary(c);
-			v = -v;
+			v.iv = -v.iv;
 			return v;
 		case '!':
 			ce_next(c);
 			v = ce_unary(c);
-			v = (v == 0);
+			v.iv = (v.uv == 0);
+			v.is_signed = 1;
 			return v;
 		case '~':
 			ce_next(c);
 			v = ce_unary(c);
-			v = ~v;
+			v.uv = ~v.uv;
 			return v;
 		case '(':
 			ce_next(c);
 			v = ce_cond(c);
 			if (c->token != ')')
-				c->result = CE_RESULT_syntax_error;
+				c->syntax_error = 1;
 			ce_next(c);
 			return v;
 		default:
@@ -392,257 +397,292 @@ static int ce_unary(ppcexp* c)
 	}
 }
 
-static int ce_mul(ppcexp* c)
+static Value ce_mul(ppcexp* c)
 {
-	int left = ce_unary(c);
+	Value left = ce_unary(c);
 	while (c->token == '*' || c->token == '/' || c->token == '%')
 	{
 		int token = c->token;
-		int right;
+		Value right;
 		ce_next(c);
 		right = ce_unary(c);
 		if (token == '*')
-		{
-			intmax_t result_wide = (intmax_t)left * (intmax_t)right;
-			if (result_wide > INT_MAX || result_wide < INT_MIN)
-			{
-				c->result = CE_RESULT_overflow;
-			}
-			left = (int)result_wide;
-		}
+			left.uv *= right.uv;
 		else
 		{
-			if (right == 0)
+			if (right.uv == 0)
 			{
-				if (c->result == CE_RESULT_ok)
-				{
-					c->result = CE_RESULT_division_by_zero;
-				}
+				if (c->syntax_error == 0)
+					c->syntax_error = -1;
 			}
+			else if (left.is_signed & right.is_signed)
+				if (token == '/')
+					left.iv /= right.iv;
+				else
+					left.iv %= right.iv;
 			else if (token == '/')
-			{
-				left /= right;
-			}
+				left.uv /= right.uv;
 			else
-			{
-				left %= right;
-			}
+				left.uv %= right.uv;
+			left.is_signed &= right.is_signed;
 		}
 	}
 	return left;
 }
 
-static int ce_sum(ppcexp* c)
+static Value ce_sum(ppcexp* c)
 {
-	int left = ce_mul(c);
+	Value left = ce_mul(c);
 	while (c->token == '+' || c->token == '-')
 	{
 		int token = c->token;
-		int right;
-		intmax_t result_wide;
+		Value right;
 		ce_next(c);
 		right = ce_mul(c);
 		if (token == '+')
-		{
-			result_wide = (intmax_t)left + (intmax_t)right;
-		}
+			left.uv += right.uv;
 		else
-		{
-			result_wide = (intmax_t)left - (intmax_t)right;
-		}
-		if (result_wide > INT_MAX || result_wide < INT_MIN)
-		{
-			c->result = CE_RESULT_overflow;
-		}
-		left = (int)result_wide;
+			left.uv -= right.uv;
+		left.is_signed &= right.is_signed;
 	}
 	return left;
 }
 
-static int ce_shift(ppcexp* c)
+static Value ce_shift(ppcexp* c)
 {
-	int left = ce_sum(c);
+	Value left = ce_sum(c);
 	while (c->token == PPLEX_shift_left || c->token == PPLEX_shift_right)
 	{
 		int tok = c->token;
-		int right;
+		Value right;
 		ce_next(c);
 		right = ce_sum(c);
 		if (tok == PPLEX_shift_right)
-		{
-			left = (left >> right);
-		}
+			if (left.is_signed)
+				left.iv = (left.iv >> right.uv);
+			else
+				left.uv = (left.uv >> right.uv);
 		else
-		{
-			left = (left << right);
-		}
+			left.uv = (left.uv << right.uv);
 	}
 	return left;
 }
 
-static int ce_compare_inequality(ppcexp* c)
+static Value ce_compare_inequality(ppcexp* c)
 {
-	int left = ce_shift(c);
+	Value left = ce_shift(c);
 	for (;;)
 	{
-		int right;
+		Value right;
 		switch (c->token)
 		{
 			case PPLEX_less_eq:
 				ce_next(c);
 				right = ce_shift(c);
-				left = (left <= right);
+				left.iv = (left.is_signed & right.is_signed) ? (left.iv <= right.iv) : (left.uv <= right.uv);
 				break;
 			case PPLEX_greater_eq:
 				ce_next(c);
 				right = ce_shift(c);
-				left = (left >= right);
+				left.iv = (left.is_signed & right.is_signed) ? (left.iv >= right.iv) : (left.uv >= right.uv);
 				break;
 			case '<':
 				ce_next(c);
 				right = ce_shift(c);
-				left = (left < right);
+				left.iv = (left.is_signed & right.is_signed) ? (left.iv < right.iv) : (left.uv < right.uv);
 				break;
 			case '>':
 				ce_next(c);
 				right = ce_shift(c);
-				left = (left > right);
+				left.iv = (left.is_signed & right.is_signed) ? (left.iv > right.iv) : (left.uv > right.uv);
 				break;
 			default:
 				return left;
 		}
+		left.is_signed = 1;
 	}
 }
 
-static int ce_compare_equality(ppcexp* c)
+static Value ce_compare_equality(ppcexp* c)
 {
-	int left = ce_compare_inequality(c), right;
+	Value left = ce_compare_inequality(c), right;
 	while (c->token == PPLEX_eq_eq || c->token == PPLEX_bang_eq)
 	{
 		int token = c->token;
 		ce_next(c);
 		right = ce_compare_inequality(c);
-		left = (left == right) ^ (token == PPLEX_bang_eq);
+		left.uv = (left.uv == right.uv) ^ (token == PPLEX_bang_eq);
+		left.is_signed = 1;
 	}
 	return left;
 }
 
-static int ce_bitwise_and(ppcexp* c)
+static Value ce_bitwise_and(ppcexp* c)
 {
-	int left = ce_compare_equality(c), right;
+	Value left = ce_compare_equality(c), right;
 	while (c->token == '&')
 	{
 		ce_next(c);
 		right = ce_compare_equality(c);
-		left &= right;
+		left.uv &= right.uv;
+		left.is_signed &= right.is_signed;
 	}
 	return left;
 }
 
-static int ce_bitwise_xor(ppcexp* c)
+static Value ce_bitwise_xor(ppcexp* c)
 {
-	int left = ce_bitwise_and(c), right;
+	Value left = ce_bitwise_and(c), right;
 	while (c->token == '^')
 	{
 		ce_next(c);
 		right = ce_bitwise_and(c);
-		left ^= right;
+		left.uv ^= right.uv;
+		left.is_signed &= right.is_signed;
 	}
 	return left;
 }
 
-static int ce_bitwise_or(ppcexp* c)
+static Value ce_bitwise_or(ppcexp* c)
 {
-	int left = ce_bitwise_xor(c), right;
+	Value left = ce_bitwise_xor(c), right;
 	while (c->token == '|')
 	{
 		ce_next(c);
 		right = ce_bitwise_xor(c);
-		left |= right;
+		left.uv |= right.uv;
+		left.is_signed &= right.is_signed;
 	}
 	return left;
 }
 
-static int ce_logical_and(ppcexp* c)
+static Value ce_logical_and(ppcexp* c)
 {
-	int left = ce_bitwise_or(c), right;
+	Value left = ce_bitwise_or(c), right;
 	while (c->token == PPLEX_and_and)
 	{
-		int result = c->result;
+		int err = c->syntax_error;
 		ce_next(c);
 		right = ce_bitwise_or(c);
-		if (left == 0 && c->result == CE_RESULT_division_by_zero && result == CE_RESULT_ok)
-			c->result = CE_RESULT_ok;  // suppress unevaluated div by 0
-		left = left && right;
+		if (left.uv == 0 && c->syntax_error < 0 && err == 0)
+			c->syntax_error = 0;  // suppress unevaluated div by 0
+		left.uv = left.uv && right.uv;
+		left.is_signed = 1;
 	}
 	return left;
 }
 
-static int ce_logical_or(ppcexp* c)
+static Value ce_logical_or(ppcexp* c)
 {
-	int left = ce_logical_and(c), right;
+	Value left = ce_logical_and(c), right;
 	while (c->token == PPLEX_or_or)
 	{
-		int result = c->result;
+		int err = c->syntax_error;
 		ce_next(c);
 		right = ce_logical_and(c);
-		if (left != 0 && c->result == CE_RESULT_division_by_zero && result == CE_RESULT_ok)
-			c->result = CE_RESULT_ok;  // suppress unevaluated div by 0
-		left = left || right;
+		if (left.uv != 0 && c->syntax_error < 0 && err == 0)
+			c->syntax_error = 0;  // suppress unevaluated div by 0
+		left.uv = left.uv || right.uv;
+		left.is_signed = 1;
 	}
 	return left;
 }
 
-static int ce_cond(ppcexp* c)
+static Value ce_cond(ppcexp* c)
 {
-	int cond = ce_logical_or(c);
+	Value cond = ce_logical_or(c);
 	if (c->token != '?')
 		return cond;
 	else
 	{
-		int v, result, left;
-		result = c->result;
+		int err;
+		Value result, left;
+		err = c->syntax_error;
 		ce_next(c);
 		left = ce_cond(c);
 		if (c->token != ':')
 		{
-			c->result = CE_RESULT_syntax_error;
+			c->syntax_error = 1;
 			return cond;
 		}
-		if (cond == 0 && c->result == CE_RESULT_division_by_zero && result == CE_RESULT_ok)
-			c->result = CE_RESULT_ok;  // suppress unevaluated div by 0
-		result = c->result;
+		if (cond.uv == 0 && c->syntax_error < 0 && err == 0)
+			c->syntax_error = 0;  // suppress unevaluated div by 0
+		err = c->syntax_error;
 		ce_next(c);
-		v = ce_cond(c);
-		if (cond != 0 && c->result == CE_RESULT_division_by_zero && result == CE_RESULT_ok)
-			c->result = CE_RESULT_ok;  // suppress unevaluated div by 0
-		v = cond != 0 ? left : v;
-		return v;
+		result = ce_cond(c);
+		if (cond.uv != 0 && c->syntax_error < 0 && err == 0)
+			c->syntax_error = 0;  // suppress unevaluated div by 0
+		result.uv = cond.uv != 0 ? left.uv : result.uv;
+		result.is_signed &= left.is_signed;
+		return result;
 	}
 }
 
-int evaluate_integer_constant_expression_as_condition(char* p, int* result)
+int evaluate_integer_constant_expression_as_condition(char* p, int* syntax_error)
 {
-	int v;
+	Value v;
 	ppcexp c = {0};
 	c.p = p;
 	ce_next(&c);
 	v = ce_cond(&c);
 	if (c.token != PPLEX_eof)
-		c.result = CE_RESULT_syntax_error;
-	*result = c.result;
-	return v != 0;
+		c.syntax_error = 1;
+	*syntax_error = c.syntax_error;
+	return v.uv != 0;
 }
 
-int evaluate_integer_constant_expression(char* p, int* result)
+#ifdef _MSC_VER
+__int64
+#else
+long long
+#endif
+evaluate_integer_constant_expression(char* p, int* syntax_error)
 {
-	int v;
+	Value v;
 	ppcexp c = {0};
 	c.p = p;
 	ce_next(&c);
 	v = ce_cond(&c);
 	if (c.token != PPLEX_eof)
-		c.result = CE_RESULT_syntax_error;
-	*result = c.result;
-	return v;
+		c.syntax_error = 1;
+	*syntax_error = c.syntax_error;
+	return v.iv;
+}
+
+#include <assert.h>
+void test_integer_constant_expression(void)
+{
+	int err;
+	assert(evaluate_integer_constant_expression("4+3*2", &err) == 10);
+	assert(err == 0);
+	assert(evaluate_integer_constant_expression("4+3*2==10", &err) == 1);
+	assert(err == 0);
+	assert(evaluate_integer_constant_expression("0xffffffff+1==1<<32", &err) == 1);
+	assert(err == 0);
+	evaluate_integer_constant_expression("1+=1", &err);
+	assert(err > 0);
+	evaluate_integer_constant_expression("1+(", &err);
+	assert(err > 0);
+	evaluate_integer_constant_expression("(1+", &err);
+	assert(err > 0);
+	assert(evaluate_integer_constant_expression("2/1", &err) == 2);
+	assert(err == 0);
+	evaluate_integer_constant_expression("1/0", &err);
+	assert(err < 0);
+	evaluate_integer_constant_expression("1 || 1/0", &err);
+	assert(err == 0);
+	evaluate_integer_constant_expression("0 || 1/0", &err);
+	assert(err < 0);
+	evaluate_integer_constant_expression("1 && 1/0", &err);
+	assert(err < 0);
+	evaluate_integer_constant_expression("0 && 1/0", &err);
+	assert(err == 0);
+	evaluate_integer_constant_expression("0 ? 1/0 : 2/1", &err);
+	assert(err == 0);
+	evaluate_integer_constant_expression("1 ? 1/0 : 2/1", &err);
+	assert(err < 0);
+	evaluate_integer_constant_expression("0 ? 2/1 : 1/0", &err);
+	assert(err < 0);
+	evaluate_integer_constant_expression("1 ? 2/1 : 1/0", &err);
+	assert(err == 0);
 }
diff --git a/Engine/Source/Developer/ShaderPreprocessor/Private/stb_preprocess/cond_expr.h b/Engine/Source/Developer/ShaderPreprocessor/Private/stb_preprocess/cond_expr.h
--- a/Engine/Source/Developer/ShaderPreprocessor/Private/stb_preprocess/cond_expr.h
+++ b/Engine/Source/Developer/ShaderPreprocessor/Private/stb_preprocess/cond_expr.h
@@ -2,19 +2,11 @@
 
 #pragma once
 
-enum
-{
-	CE_RESULT_ok,
-	CE_RESULT_syntax_error,
-	CE_RESULT_division_by_zero,
-	CE_RESULT_overflow,
-};
-
-#ifdef __cplusplus
-#define CONDEXPR_DEF extern "C"
+#ifdef _MSC_VER
+__int64
 #else
-#define CONDEXPR_DEF extern
+long long
 #endif
-
-CONDEXPR_DEF int evaluate_integer_constant_expression(char* p, int* error);
-CONDEXPR_DEF int evaluate_integer_constant_expression_as_condition(char* p, int* error);
\ No newline at end of file
+evaluate_integer_constant_expression(char* p, int* syntax_error);
+extern int evaluate_integer_constant_expression_as_condition(char* p, int* syntax_error);
+extern void test_integer_constant_expression(void);
diff --git a/Engine/Source/Developer/ShaderPreprocessor/Private/stb_preprocess/preprocessor.c b/Engine/Source/Developer/ShaderPreprocessor/Private/stb_preprocess/preprocessor.c
--- a/Engine/Source/Developer/ShaderPreprocessor/Private/stb_preprocess/preprocessor.c
+++ b/Engine/Source/Developer/ShaderPreprocessor/Private/stb_preprocess/preprocessor.c
@@ -4328,36 +4328,20 @@ static void process_directive(parse_state* cs, conditional_state* cons)
 		case HASH_if:
 		{
 			ifdef_info ifinfo;
-			int result, result_code;
+			int result, syntax_error;
 
 			ifinfo.line = cs->src_line_number;
 			ifinfo.type = HASH_if;
 
 			p = preprocessor_skip_whitespace_simple(p);
-			result = evaluate_if(cs, p, &result_code);
+			result = evaluate_if(cs, p, &syntax_error);
 
-			if (result_code != CE_RESULT_ok)
+			if (syntax_error)
 			{
-				char* err_msg;
-				int ce_result;
-				switch (result_code)
-				{
-				case CE_RESULT_division_by_zero:
-					err_msg = "Division by zero occurred when evaluating conditional expression";
-					ce_result = PP_RESULT_conditional_expression_division_by_zero;
-					break;
-				case CE_RESULT_overflow:
-					err_msg = "Overflow occurred when evaluating conditional expression";
-					ce_result = PP_RESULT_conditional_expression_overflow;
-					break;
-				case CE_RESULT_syntax_error:
-					err_msg = "Syntax error occurred when evaluating conditional expression";
-					ce_result = PP_RESULT_syntax_error_conditional_expression;
-					break;
-				}
-				
-				do_error_explicit(cs, ce_result, cs->src_line_number, err_msg);
-				return;
+				if (do_error(cs, "Syntax error in conditional expression of '#if'"))
+					return;
+				else
+					result = 0;
 			}
 
 			arrpush(c->ifdef_stack, ifinfo);
@@ -5228,6 +5212,7 @@ void init_preprocessor(
 
 	init_preprocessor_scanner();
 	init_char_type();
+	// test_integer_constant_expression();
 
 	init_directive("define", HASH_define);
 	init_directive("else", HASH_else);
diff --git a/Engine/Source/Developer/ShaderPreprocessor/Private/stb_preprocess/preprocessor.h b/Engine/Source/Developer/ShaderPreprocessor/Private/stb_preprocess/preprocessor.h
--- a/Engine/Source/Developer/ShaderPreprocessor/Private/stb_preprocess/preprocessor.h
+++ b/Engine/Source/Developer/ShaderPreprocessor/Private/stb_preprocess/preprocessor.h
@@ -26,9 +26,6 @@ enum
 	PP_RESULT_malformed_pragma_message,
 	PP_RESULT_malformed_line_filename,
 	PP_RESULT_invalid_line_linenumber,
-	PP_RESULT_syntax_error_conditional_expression,
-	PP_RESULT_conditional_expression_overflow,
-	PP_RESULT_conditional_expression_division_by_zero,
 
 	PP_RESULT_count
 };
diff --git a/Engine/Source/Developer/ShaderPreprocessor/Tests/ShaderPreprocessorTests.cpp b/Engine/Source/Developer/ShaderPreprocessor/Tests/ShaderPreprocessorTests.cpp
deleted file mode 100644
diff --git a/Engine/Source/Programs/LowLevelTests/FoundationTests/FoundationTests.Build.cs b/Engine/Source/Programs/LowLevelTests/FoundationTests/FoundationTests.Build.cs
--- a/Engine/Source/Programs/LowLevelTests/FoundationTests/FoundationTests.Build.cs
+++ b/Engine/Source/Programs/LowLevelTests/FoundationTests/FoundationTests.Build.cs
@@ -55,8 +55,7 @@ public class FoundationTests : TestModuleRules
 		{
 			PrivateDependencyModuleNames.AddRange(
 				new string[] {
-					"DesktopPlatform",
-					"ShaderPreprocessor"
+					"DesktopPlatform"
 				});
 		}
 	}
