[shaders] remove usage of sprintf from preprocessor (along with some minor cleanup of the code around this usage), add some unit tests for affected code

diff --git a/Engine/Source/Developer/ShaderPreprocessor/Private/stb_preprocess/preprocessor.c b/Engine/Source/Developer/ShaderPreprocessor/Private/stb_preprocess/preprocessor.c
--- a/Engine/Source/Developer/ShaderPreprocessor/Private/stb_preprocess/preprocessor.c
+++ b/Engine/Source/Developer/ShaderPreprocessor/Private/stb_preprocess/preprocessor.c
@@ -374,18 +374,24 @@ static char* find_closest_matching_string(char* reference, char* choices[], int
 		return NULL;
 }
 
-static char* error_char(char ch, char buffer[12])
+typedef struct
+{
+	char buf[12];
+} err_charbuf12;
+
+static err_charbuf12 error_char(char ch)
 {
+	err_charbuf12 b;
 	if (ch >= 32 && ch < 127)
 		if (ch == '\'' || ch == '\\')
-			sprintf(buffer, "'\\%c'", ch);
+			snprintf(b.buf, sizeof(b.buf), "'\\%c'", ch);
 		else
-			sprintf(buffer, "'%c'", ch);
+			snprintf(b.buf, sizeof(b.buf), "'%c'", ch);
 	else if (ch == 0)
-		sprintf(buffer, "EOF");
+		snprintf(b.buf, sizeof(b.buf), "EOF");
 	else
-		sprintf(buffer, "value %d", ch);  // "value -100\0" = 11 (6 + 1 + 3 + 1)
-	return buffer;
+		snprintf(b.buf, sizeof(b.buf), "value %d", ch);  // "value -100\0" = 11 (6 + 1 + 3 + 1)
+	return b;
 }
 
 // write a line number in 12 characters of space, and return a pointer after the number
@@ -393,10 +399,11 @@ static char* error_char(char ch, char buffer[12])
 //     1. 1-digit numbers take 3 characters with two leading spaces
 //     2. 2-digit numbers take 2 characters with one leading space
 //     3. all other numbers have no extra spaces
-//     4. sprintf() is used if number is > 9999, on the assumption this is rare
+//     4. snprintf() is used if number is > 9999, on the assumption this is rare
 //            -- but note that if a file is, say, 500K lines, then there may be
 //               very few FILES that use this path, but an inordinate number of lines
-static char* preprocessor_line_dtoa(char q[12], int number)
+#define INT_STR_SIZE 12
+static char* preprocessor_line_dtoa(char q[INT_STR_SIZE], int number)
 {
 	static const char digits[10 + 1] = "0123456789";
 	static const char digit_pairs[200 + 1] =
@@ -411,12 +418,10 @@ static char* preprocessor_line_dtoa(char q[12], int number)
 		q[0] = first_digit[leading];
 		q[1] = digit_pairs[n * 2 + 0];
 		q[2] = digit_pairs[n * 2 + 1];
-#if 1
 		if (number < 10)
+		{
 			q[1] = ' ';
-#else
-		q[(number < 10) ? 1 : 3] = ' ';
-#endif
+		}
 		return q + 3;
 	}
 	else if (number <= 9999)
@@ -431,7 +436,7 @@ static char* preprocessor_line_dtoa(char q[12], int number)
 	}
 	else
 	{
-		return q + sprintf(q, "%d", number);
+		return q + snprintf(q, INT_STR_SIZE, "%d", number);
 	}
 }
 
@@ -442,7 +447,12 @@ static void output_line_directive(parse_state* cs)
 
 	// compute an upper bound on the amount of data to be written
 	size_t filename_len = strlen(cs->filename);
-	size_t expansion_length = 6 + 16 + filename_len + 2 + 2 + 1;
+	size_t expansion_length =	1 +					// possible newline
+								6 +					// "#line "
+								INT_STR_SIZE +	// max size of line number string
+								2 +					// space + open quote for filename
+								filename_len +
+								2;					// close quote for filename + newline
 
 	// make sure there's room in the output for the #line directive plus all the remaining text
 	arrsetcap(out, arrlen(out) + (size_t)(cs->src_length - cs->src_offset) + expansion_length);
@@ -579,7 +589,9 @@ static int do_error_code(parse_state* ps, int code, char* text, ...)
 	va_list va;
 
 	if (pp_result_mode[code] == PP_RESULT_MODE_no_warning)
+	{
 		return 0;
+	}
 
 	va_start(va, text);
 	err = error_explicit_v(ps, code, 0, text, va);
@@ -587,6 +599,8 @@ static int do_error_code(parse_state* ps, int code, char* text, ...)
 	return err;
 }
 
+#define do_error_char_code(ps, code, text, c, ...) do_error_code(ps, code, text, error_char(c).buf, ##__VA_ARGS__)
+
 static int do_error(parse_state* ps, char* text, ...)
 {
 	int err;
@@ -598,6 +612,8 @@ static int do_error(parse_state* ps, char* text, ...)
 	return err;
 }
 
+#define do_error_char(ps, text, c, ...) do_error(ps, text, error_char(c).buf, ##__VA_ARGS__)
+
 static int do_error_explicit(parse_state* ps, int code, int line_number, char* text, ...)
 {
 	int err;
@@ -2783,16 +2799,13 @@ error_exit:
 
 invalid_char_error:
 {
-	char cbuf[12];
 	char* error_messages[3] = {
 		"Invalid character %s after '#define'",
 		"Invalid character %s in macro definition after '%s'",
 		"Invalid character %s in parameter list for macro definition '%s'",
 	};
 
-	char bad = *p;
-
-	do_error_code(ps, PP_RESULT_invalid_char, error_messages[state], error_char(bad, cbuf), m->symbol_name);
+	do_error_char_code(ps, PP_RESULT_invalid_char, error_messages[state], *p, m->symbol_name);
 
 	return NULL;
 }
@@ -3130,8 +3143,7 @@ static void maybe_expand_macro(parse_state* cs, struct macro_definition* pending
 						++p;
 					if (*p != ')')
 					{
-						char cbuf[12];
-						do_error(cs, "unexpected character %s at end of 'defined('", error_char(*p, cbuf));
+						do_error_char(cs, "unexpected character %s at end of 'defined('", *p);
 						return;
 					}
 					++p;
@@ -3156,8 +3168,8 @@ static void maybe_expand_macro(parse_state* cs, struct macro_definition* pending
 			{
 				int len;
 				q = identifier;
-				arrsetcap(cs->dest, arrlen(cs->dest) + 12);
-				len = sprintf(cs->dest + arrlen(cs->dest), "%d", cs->src_line_number);
+				arrsetcap(cs->dest, arrlen(cs->dest) + INT_STR_SIZE);
+				len = snprintf(cs->dest + arrlen(cs->dest), INT_STR_SIZE, "%d", cs->src_line_number);
 				arrsetlen(cs->dest, arrlen(cs->dest) + len);
 				break;
 			}
@@ -3171,8 +3183,8 @@ static void maybe_expand_macro(parse_state* cs, struct macro_definition* pending
 				}
 				int len;
 				q = identifier;
-				arrsetcap(cs->dest, arrlen(cs->dest) + 12);
-				len = sprintf(cs->dest + arrlen(cs->dest), "%d", c->counter);
+				arrsetcap(cs->dest, arrlen(cs->dest) + INT_STR_SIZE);
+				len = snprintf(cs->dest + arrlen(cs->dest), INT_STR_SIZE, "%d", c->counter);
 				arrsetlen(cs->dest, arrlen(cs->dest) + len);
 				++c->counter;
 				break;
@@ -3714,8 +3726,7 @@ static const char* parse_pp_identifier_in_directive(parse_state* ps, const char*
 	identifier[0] = 0;
 	if (!char_is_pp_identifier_first(*p))
 	{
-		char cbuf[12];
-		do_error_code(ps, PP_RESULT_invalid_char, "Invalid character %s after '#%s'", error_char(*p, cbuf), directive);
+		do_error_char_code(ps, PP_RESULT_invalid_char, "Invalid character %s after '#%s'", *p, directive);
 		return p;
 	}
 	identifier[0] = *p++;
@@ -3728,7 +3739,9 @@ static const char* parse_pp_identifier_in_directive(parse_state* ps, const char*
 	*idlen = i;
 	identifier[i] = 0;
 	if (char_is_pp_identifier(*p))
+	{
 		do_error_code(ps, PP_RESULT_identifier_too_long, "'#%s' identifier too long", directive);
+	}
 	return p;
 }
 
@@ -3761,8 +3774,7 @@ static int process_include(parse_state* cs, const char* start, conditional_state
 	// @TODO: macro expansion
 	if (!(*p == '<' || *p == '"'))
 	{
-		char cbuf[12];
-		do_error(cs, "Expected filename after '#include', found character %s", error_char(*p, cbuf));
+		do_error_char(cs, "Expected filename after '#include', found character %s", *p);
 		return RETURN_BEHAVIOR_return;
 	}
 	else
@@ -3799,8 +3811,7 @@ static int process_include(parse_state* cs, const char* start, conditional_state
 		p = preprocessor_skip_whitespace_simple(p + 1);
 		if (*p != 0)
 		{
-			char cbuf[12];
-			do_error(cs, "Unexpected character %s after filename in '#include'", error_char(*p, cbuf));
+			do_error_char(cs, "Unexpected character %s after filename in '#include'", *p);
 			return RETURN_BEHAVIOR_return;
 		}
 
@@ -3924,8 +3935,7 @@ static void process_directive(parse_state* cs, conditional_state* cons)
 					// invalid directive, produce a diagnostic
 					if (!char_is_pp_identifier_first(*p))
 					{
-						char cbuf[12];
-						do_error_code(cs, PP_RESULT_invalid_char, "Invalid character %s after '#'", error_char(*p, cbuf));
+						do_error_char_code(cs, PP_RESULT_invalid_char, "Invalid character %s after '#'", *p);
 					}
 					else
 					{
@@ -4283,9 +4293,7 @@ static void process_directive(parse_state* cs, conditional_state* cons)
 			p = preprocessor_skip_whitespace_simple(p);
 			if (*p != 0)
 			{
-				char cbuf[12];
-				do_error_code(cs, PP_RESULT_unexpected_character_after_end_of_directive, "Expected end of line in #ifdef, but found character %s",
-							  error_char(*p, cbuf));
+				do_error_char_code(cs, PP_RESULT_unexpected_character_after_end_of_directive, "Expected end of line in #ifdef, but found character %s", *p);
 			}
 			break;
 		}
@@ -4318,9 +4326,7 @@ static void process_directive(parse_state* cs, conditional_state* cons)
 			p = preprocessor_skip_whitespace_simple(p);
 			if (*p != 0)
 			{
-				char cbuf[12];
-				do_error_code(cs, PP_RESULT_unexpected_character_after_end_of_directive, "Expected end of line in #ifdef, but found character %s.",
-							  error_char(*p, cbuf));
+				do_error_char_code(cs, PP_RESULT_unexpected_character_after_end_of_directive, "Expected end of line in #ifdef, but found character %s.", *p);
 			}
 			break;
 		}
@@ -4427,10 +4433,10 @@ static void process_directive(parse_state* cs, conditional_state* cons)
 
 			if (*p != 0)
 			{
-				char cbuf[12];
-				if (do_error_code(cs, PP_RESULT_unexpected_character_after_end_of_directive, "Expected end of line in #else, but found character %s.",
-								  error_char(*p, cbuf)))
+				if (do_error_char_code(cs, PP_RESULT_unexpected_character_after_end_of_directive, "Expected end of line in #else, but found character %s.", *p))
+				{
 					return;
+				}
 			}
 
 			if (cons->conditional_compilation_nesting_level == 0)
@@ -4469,10 +4475,10 @@ static void process_directive(parse_state* cs, conditional_state* cons)
 			p = preprocessor_skip_whitespace_simple(p);
 			if (*p != 0)
 			{
-				char cbuf[12];
-				if (do_error_code(cs, PP_RESULT_unexpected_character_after_end_of_directive, "Expected end of line in #endif, but found character %s.",
-								  error_char(*p, cbuf)))
+				if (do_error_char_code(cs, PP_RESULT_unexpected_character_after_end_of_directive, "Expected end of line in #endif, but found character %s.", *p))
+				{
 					return;
+				}
 			}
 
 			if (cons->conditional_compilation_nesting_level == 0)
diff --git a/Engine/Source/Developer/ShaderPreprocessor/Tests/ShaderPreprocessorTests.cpp b/Engine/Source/Developer/ShaderPreprocessor/Tests/ShaderPreprocessorTests.cpp
--- a/Engine/Source/Developer/ShaderPreprocessor/Tests/ShaderPreprocessorTests.cpp
+++ b/Engine/Source/Developer/ShaderPreprocessor/Tests/ShaderPreprocessorTests.cpp
@@ -90,7 +90,12 @@ struct FPreprocessTestResult
 FPreprocessTestResult ExecutePreprocessTest(const FShaderSource::FViewType& Source)
 {
 	FPreprocessTestResult Result;
-	init_preprocessor(&TestLoadFile, &TestFreeFile, nullptr, nullptr, nullptr);
+	static bool bInitialized = false;
+	if (!bInitialized)
+	{
+		init_preprocessor(&TestLoadFile, &TestFreeFile, nullptr, nullptr, nullptr);
+		bInitialized = true;
+	}
 	Result.Source = preprocess_file("preprocessor_test", const_cast<FShaderSource::FViewType*>(&Source), nullptr, 0u, &Result.Diagnostics, &Result.NumDiagnostics, nullptr, 0u);
 	return Result;
 }
@@ -108,4 +113,112 @@ TEST_CASE_NAMED(FConditionalExpressionErrorReporting, "Shaders::ShaderPreprocess
 	CHECK(strstr(Result.Diagnostics[0].message, "Overflow") != nullptr);
 }
 
+TEST_CASE_NAMED(FInvalidCharError, "Shaders::ShaderPreprocessor::InvalidCharErrors", "[EditorContext][EngineFilter]")
+{
+	{
+		// EOF after #define
+		FShaderSource::FViewType Source = SHADER_SOURCE_VIEWLITERAL(
+			"#" "define");
+		FPreprocessTestResult Result = ExecutePreprocessTest(Source);
+		CHECK(Result.NumDiagnostics == 1);
+		CHECK(Result.Diagnostics[0].where->line_number == 1);
+		CHECK(strstr(Result.Diagnostics[0].message, "Invalid character EOF") != nullptr);
+	}
+
+	{
+		// Invalid printable char in define param list
+		FShaderSource::FViewType Source = SHADER_SOURCE_VIEWLITERAL(
+			"#" "define BLAH(abc,#) abc\n");
+		FPreprocessTestResult Result = ExecutePreprocessTest(Source);
+		CHECK(Result.NumDiagnostics == 1);
+		CHECK(Result.Diagnostics[0].where->line_number == 1);
+		CHECK(strstr(Result.Diagnostics[0].message, "Invalid character '#'") != nullptr);
+	}
+
+	{
+		// Invalid non-printable char in define param list
+		FShaderSource::FViewType Source = SHADER_SOURCE_VIEWLITERAL(
+			"#" "define BLAH(abc,\3) abc\n");
+		FPreprocessTestResult Result = ExecutePreprocessTest(Source);
+		CHECK(Result.NumDiagnostics == 1);
+		CHECK(Result.Diagnostics[0].where->line_number == 1);
+		CHECK(strstr(Result.Diagnostics[0].message, "Invalid character value 3") != nullptr);
+	}
+
+	{
+		// Unexpected char after #if defined(
+		FShaderSource::FViewType Source = SHADER_SOURCE_VIEWLITERAL(
+			"#" "define BLAH 1\n"
+			"#" "if defined(BLAH*)\n"
+			"#" "endif\n"
+		);
+		FPreprocessTestResult Result = ExecutePreprocessTest(Source);
+		CHECK(Result.NumDiagnostics == 1);
+		CHECK(Result.Diagnostics[0].where->line_number == 2);
+		CHECK(strstr(Result.Diagnostics[0].message, "unexpected character '*'") != nullptr);
+	}
+}
+
+TEST_CASE_NAMED(FSpecialMacros, "Shaders::ShaderPreprocessor::BuiltinMacros", "[EditorContext][EngineFilter]")
+{
+	FShaderSource::FViewType Source = SHADER_SOURCE_VIEWLITERAL(
+		"int A = __" "LINE__;\n"
+		"int B = __" "COUNTER__;\n"
+		"int C = __" "COUNTER__;\n"
+	);
+	FPreprocessTestResult Result = ExecutePreprocessTest(Source);
+	CHECK(Result.NumDiagnostics == 0);
+	CHECK(strstr(Result.Source, "int A = 1") != nullptr);
+	CHECK(strstr(Result.Source, "int B = 0") != nullptr);
+	CHECK(strstr(Result.Source, "int C = 1") != nullptr);
+}
+
+TEST_CASE_NAMED(FEmitLineDirective, "Shaders::ShaderPreprocessor::EmitLineDirective", "[EditorContext][EngineFilter]")
+{
+	// line directive will be emitted if 12 or more lines are removed due to an inactive preprocessor block
+	// (otherwise the inactive code will be replaced with empty lines)
+	FShaderSource::FViewType Source = SHADER_SOURCE_VIEWLITERAL(
+		"#" "define MYDIR 0\n"
+		"#" "if MYDIR\n"
+		"int A = 0;\n"
+		"#" "endif\n"
+		"#" "if MYDIR\n"
+		"int B = 0;\n"
+		"B++;\n"
+		"B++;\n"
+		"B++;\n"
+		"B++;\n"
+		"B++;\n"
+		"B++;\n"
+		"B++;\n"
+		"B++;\n"
+		"B++;\n"
+		"B++;\n"
+		"B++;\n"
+		"#" "endif\n"
+		"int C = 0;\n"
+	);
+	// exactly two line directives will be emitted, one on the first line (done for all files) and one after the second inactive block (starting with the newline on line 18)
+	FPreprocessTestResult Result = ExecutePreprocessTest(Source);
+	CHECK(Result.NumDiagnostics == 0);
+	ANSICHAR* Directive = strstr(Result.Source, "#line");
+	if (Directive)
+	{
+		CHECK(strstr(Directive, "1 \"preprocessor_test\"") != nullptr);
+		ANSICHAR* NextDirective = strstr(Directive + 5, "#line");
+		if (NextDirective)
+		{
+			CHECK(strstr(NextDirective, "18 \"preprocessor_test\"") != nullptr);
+		}
+		else
+		{
+			FAIL_CHECK("No second line directive found");
+		}
+	}
+	else
+	{
+		FAIL_CHECK("No line directive found");
+	}
+}
+
 #endif
