[shaders]
- add a configurable function for the LCG used by stb_ds hashtable and make the accumulator variable threadlocal in our implementation to fix TSAN issues; this is a workaround for the fact that there's no apparent way to use C11's _Thread_local when compiling mixed C/C++ on MSVC
- remove now-unnecessary function for explicitly setting the hash seed since it's now fully in control of the custom LCG implementation

diff --git a/Engine/Source/Developer/ShaderPreprocessor/Private/ShaderPreprocessor.cpp b/Engine/Source/Developer/ShaderPreprocessor/Private/ShaderPreprocessor.cpp
--- a/Engine/Source/Developer/ShaderPreprocessor/Private/ShaderPreprocessor.cpp
+++ b/Engine/Source/Developer/ShaderPreprocessor/Private/ShaderPreprocessor.cpp
@@ -113,6 +113,19 @@ extern "C"
 		}
 		return nullptr;
 	}
+
+	size_t StbLcgNext()
+	{
+		// thread local linear congruential generator
+		static thread_local uint64 Current = 0x31415926;
+		static const uint64 A = 0x27BB2EE687B0B0FDULL;
+		static const uint64 B = 0xB504F32DULL;
+		
+		size_t Next = Current;
+		Current = Current * A + B;
+
+		return (size_t)Next;
+	}
 }
 
 struct FStbLoadedInclude
diff --git a/Engine/Source/Developer/ShaderPreprocessor/Private/StbConfig.h b/Engine/Source/Developer/ShaderPreprocessor/Private/StbConfig.h
--- a/Engine/Source/Developer/ShaderPreprocessor/Private/StbConfig.h
+++ b/Engine/Source/Developer/ShaderPreprocessor/Private/StbConfig.h
@@ -12,6 +12,7 @@ void* StbMalloc(size_t Size);
 void* StbRealloc(void* Pointer, size_t Size);
 void StbFree(void* Pointer);
 char* StbStrDup(const char* String);
+size_t StbLcgNext();
 #ifdef __cplusplus
 }
 #endif
@@ -20,3 +21,4 @@ char* StbStrDup(const char* String);
 #define STB_COMMON_REALLOC(Pointer, Size) StbRealloc(Pointer, Size)
 #define STB_COMMON_FREE(Pointer) StbFree(Pointer)
 #define STB_COMMON_STRDUP(String) StbStrDup(String)
+#define STB_LCG_NEXT() StbLcgNext()
diff --git a/Engine/Source/Developer/ShaderPreprocessor/Private/stb_preprocess/stb_ds.c b/Engine/Source/Developer/ShaderPreprocessor/Private/stb_preprocess/stb_ds.c
--- a/Engine/Source/Developer/ShaderPreprocessor/Private/stb_preprocess/stb_ds.c
+++ b/Engine/Source/Developer/ShaderPreprocessor/Private/stb_preprocess/stb_ds.c
@@ -179,13 +179,6 @@ typedef struct
 #define STBDS_HASH_EMPTY 0
 #define STBDS_HASH_DELETED 1
 
-static size_t stbds_hash_seed = 0x31415926;
-
-void stbds_rand_seed(size_t seed)
-{
-	stbds_hash_seed = seed;
-}
-
 #define stbds_load_32_or_64(var, temp, v32, v64_hi, v64_lo)                                          \
 	temp = v64_lo ^ v32, temp <<= 16, temp <<= 16, temp >>= 16, temp >>= 16, /* discard if 32-bit */ \
 		var = v64_hi, var <<= 16, var <<= 16,								 /* discard if 32-bit */ \
@@ -281,13 +274,7 @@ static stbds_hash_index* stbds_make_hash_index(size_t slot_count, stbds_hash_ind
 	{
 		size_t a, b, temp;
 		memset(&t->string, 0, sizeof(t->string));
-		t->seed = stbds_hash_seed;
-		// LCG
-		// in 32-bit, a =          2147001325   b =  715136305
-		// in 64-bit, a = 2862933555777941757   b = 3037000493
-		stbds_load_32_or_64(a, temp, 2147001325, 0x27bb2ee6, 0x87b0b0fd);
-		stbds_load_32_or_64(b, temp, 715136305, 0, 0xb504f32d);
-		stbds_hash_seed = stbds_hash_seed * a + b;
+		t->seed = STB_LCG_NEXT();
 	}
 
 	{
diff --git a/Engine/Source/Developer/ShaderPreprocessor/Private/stb_preprocess/stb_ds.h b/Engine/Source/Developer/ShaderPreprocessor/Private/stb_preprocess/stb_ds.h
--- a/Engine/Source/Developer/ShaderPreprocessor/Private/stb_preprocess/stb_ds.h
+++ b/Engine/Source/Developer/ShaderPreprocessor/Private/stb_preprocess/stb_ds.h
@@ -453,9 +453,6 @@ extern "C"
 {
 #endif
 
-	// for security against attackers, seed the library with a random number, at least time() but stronger is better
-	extern void stbds_rand_seed(size_t seed);
-
 	// these are the hash functions used internally if you want to test them or use them for other purposes
 	extern size_t stbds_hash_bytes(void* p, size_t len, size_t seed);
 	extern size_t stbds_hash_string(char* str, size_t seed);
