update header documentation per bughawk feedback

diff --git a/Engine/Source/Developer/ShaderPreprocessor/Private/stb_preprocess/stb_ds.h b/Engine/Source/Developer/ShaderPreprocessor/Private/stb_preprocess/stb_ds.h
--- a/Engine/Source/Developer/ShaderPreprocessor/Private/stb_preprocess/stb_ds.h
+++ b/Engine/Source/Developer/ShaderPreprocessor/Private/stb_preprocess/stb_ds.h
@@ -141,11 +141,6 @@ DOCUMENTATION
 	Special interface:
 
 
-	  stbds_rand_seed:
-		void stbds_rand_seed(size_t seed);
-		  For security against adversarially chosen data, you should seed the
-		  library with a strong random number. Or at least seed it with time().
-
 	  stbds_hash_string:
 		size_t stbds_hash_string(char *str, size_t seed);
 		  Returns a hash value for a string.
@@ -345,8 +340,8 @@ NOTES - HASH MAP
 	this can just replace current arena implementation.
 
   * If adversarial input is a serious concern and you're on a 64-bit platform,
-	enable STBDS_SIPHASH_2_4 (see the 'Compile-time options' section), and pass
-	a strong random number to stbds_rand_seed.
+	enable STBDS_SIPHASH_2_4 (see the 'Compile-time options' section), and ensure
+	values returned by STB_LCG_NEXT function implementation are strong random numbers.
 
   * The default value for the hash table is stored in foo[-1], so if you
 	use code like 'hmget(T,k)->value = 5' you can accidentally overwrite
