Shader preprocessor test - LLT only

#rnx

#Synced-CL 0

diff --git a/Engine/Source/Developer/ShaderPreprocessor/Tests/ShaderPreprocessorTests.cpp b/Engine/Source/Developer/ShaderPreprocessor/Tests/ShaderPreprocessorTests.cpp
--- a/Engine/Source/Developer/ShaderPreprocessor/Tests/ShaderPreprocessorTests.cpp
+++ b/Engine/Source/Developer/ShaderPreprocessor/Tests/ShaderPreprocessorTests.cpp
@@ -1,18 +1,18 @@
 // Copyright Epic Games, Inc. All Rights Reserved.
 
-#if WITH_TESTS
+#if WITH_LOW_LEVEL_TESTS
 
 #include "stb_preprocess/cond_expr.h"
 #include "stb_preprocess/preprocessor.h"
 #include "ShaderSource.h"
-#include "Tests/TestHarnessAdapter.h"
+#include "TestHarness.h"
 
 int TestEvaluate(const char* Expression, int& Result)
 {
 	return evaluate_integer_constant_expression(const_cast<char*>(Expression), &Result);
 }
 
-TEST_CASE_NAMED(FConditionalExpressionTests, "Shaders::ShaderPreprocessor::ConditionalExpression", "[EditorContext][EngineFilter]")
+TEST_CASE("Shaders::ShaderPreprocessor::ConditionalExpression", "[EditorContext][EngineFilter]")
 {
 	int Result;
 	CHECK(TestEvaluate("4+3*2", Result) == 10);
@@ -100,7 +100,7 @@ FPreprocessTestResult ExecutePreprocessTest(const FShaderSource::FViewType& Sour
 	return Result;
 }
 
-TEST_CASE_NAMED(FConditionalExpressionErrorReporting, "Shaders::ShaderPreprocessor::ConditionalExpressionErrorReporting", "[EditorContext][EngineFilter]")
+TEST_CASE("Shaders::ShaderPreprocessor::ConditionalExpressionErrorReporting", "[EditorContext][EngineFilter]")
 {
 	// need this dance of multiple assembled string literal fragments to avoid the C++ preprocessor parsing directives inside the literal
 	FShaderSource::FViewType Source = ANSITEXTVIEW(
@@ -113,7 +113,7 @@ TEST_CASE_NAMED(FConditionalExpressionErrorReporting, "Shaders::ShaderPreprocess
 	CHECK(strstr(Result.Diagnostics[0].message, "Overflow") != nullptr);
 }
 
-TEST_CASE_NAMED(FInvalidCharError, "Shaders::ShaderPreprocessor::InvalidCharErrors", "[EditorContext][EngineFilter]")
+TEST_CASE("Shaders::ShaderPreprocessor::InvalidCharErrors", "[EditorContext][EngineFilter]")
 {
 	{
 		// EOF after #define
@@ -159,7 +159,7 @@ TEST_CASE_NAMED(FInvalidCharError, "Shaders::ShaderPreprocessor::InvalidCharErro
 	}
 }
 
-TEST_CASE_NAMED(FSpecialMacros, "Shaders::ShaderPreprocessor::BuiltinMacros", "[EditorContext][EngineFilter]")
+TEST_CASE("Shaders::ShaderPreprocessor::BuiltinMacros", "[EditorContext][EngineFilter]")
 {
 	FShaderSource::FViewType Source = ANSITEXTVIEW(
 		"int A = __" "LINE__;\n"
@@ -173,7 +173,7 @@ TEST_CASE_NAMED(FSpecialMacros, "Shaders::ShaderPreprocessor::BuiltinMacros", "[
 	CHECK(strstr(Result.Source, "int C = 1") != nullptr);
 }
 
-TEST_CASE_NAMED(FEmitLineDirective, "Shaders::ShaderPreprocessor::EmitLineDirective", "[EditorContext][EngineFilter]")
+TEST_CASE("Shaders::ShaderPreprocessor::EmitLineDirective", "[EditorContext][EngineFilter]")
 {
 	// line directive will be emitted if 12 or more lines are removed due to an inactive preprocessor block
 	// (otherwise the inactive code will be replaced with empty lines)
