[shaders] preprocessor - fix incorrect length calculation in macro arguments. if we're preserving leading whitespace (in the case of multi-line macro argument lists) we were incorrectly truncating the result array because we weren't accounting for the re-added whitespace characters.

#Synced-CL 51095924

diff --git a/Engine/Source/Developer/ShaderPreprocessor/Private/stb_preprocess/preprocessor.c b/Engine/Source/Developer/ShaderPreprocessor/Private/stb_preprocess/preprocessor.c
--- a/Engine/Source/Developer/ShaderPreprocessor/Private/stb_preprocess/preprocessor.c
+++ b/Engine/Source/Developer/ShaderPreprocessor/Private/stb_preprocess/preprocessor.c
@@ -3400,12 +3400,20 @@ static void maybe_expand_macro(parse_state* cs, struct macro_definition* pending
 
 			{
 				char* e;
-				// if this is a single line argument, need to strip trailing whitespace so token-pasting works
-				// otherwise we keep the whitespace as-is by re-inserting the leading ws, and so assume multiline macro arguments are not token pasted
-				arrput(copy, 0);
+				// if this is a multi-line argument maintain the whitespace by re-inserting the leading ws (this assumes multiline macro arguments are not token pasted)
+				if (arg_newlines > 0)
+				{
+					arrinsn(copy, 0, leading_ws_size); //-V568
+					memcpy(copy, leading_ws, leading_ws_size);
+
+				}
+				arrput(copy, 0); // add null terminator to working array
+
 				p = copy;
-				e = copy + arrlen(copy) - 1;  // get address of NUL
+				e = copy + arrlen(copy) - 1;  // get address of null terminator
+				ptrdiff_t len = e - p;			  // get length excluding null terminator
 				
+				// if this is a single line argument, need to strip trailing whitespace so token-pasting works
 				if (arg_newlines == 0)
 				{
 					if (e - 1 > p)
@@ -3414,11 +3422,7 @@ static void maybe_expand_macro(parse_state* cs, struct macro_definition* pending
 						assert(e > p);
 					}
 					*e = 0;
-				}
-				else
-				{
-					arrinsn(copy, 0, leading_ws_size); //-V568
-					memcpy(copy, leading_ws, leading_ws_size);
+					len = e - p; // reset new length after stripping trailing ws
 				}
 
 #if SSE_READ_PADDING
@@ -3442,7 +3446,7 @@ static void maybe_expand_macro(parse_state* cs, struct macro_definition* pending
 				}
 #endif
 
-				arrsetlen(copy, e - p);
+				arrsetlen(copy, len); // set argument string len to exclude null terminator
 				arrput(arguments, copy);
 				cs->src_line_number += arg_newlines;
 				cs->dest_line_number += arg_newlines;
