[shaders/preprocessor] fix stomp in preprocessing when the end of the inline buffer is hit when adding SSE padding via arrsetcap in copy_to_action_point and copy_to_action_point_macro_expansion. this capacity increase could potentially reallocate the buffer, causing the cs->dest output pointer to become stale - and this pointer is used to set the output length in the subsequently called copy_and_filter_macro function. in the case where this buffer allocates, and the length set is required, this would write 4 bytes into a freed buffer.

#Synced-CL 51590704

diff --git a/Engine/Source/Developer/ShaderPreprocessor/Private/stb_preprocess/preprocessor.c b/Engine/Source/Developer/ShaderPreprocessor/Private/stb_preprocess/preprocessor.c
--- a/Engine/Source/Developer/ShaderPreprocessor/Private/stb_preprocess/preprocessor.c
+++ b/Engine/Source/Developer/ShaderPreprocessor/Private/stb_preprocess/preprocessor.c
@@ -1477,6 +1477,7 @@ static int copy_to_action_point(parse_state* cs)
 		// expand output so there's room to copy entire file unmodified, plus padding so SSE operations can write to the end of the buffer without special cases
 		int out_len = arrlen(out);
 		arrsetcap(out, out_len + (cs->src_length - cs->src_offset) + SSE_READ_PADDING);
+		cs->dest = out;  // arrsetcap may realloc 'out' to a new address; keep cs->dest in sync so copy_and_filter_macro uses a valid pointer
 		q = out + out_len;
 
 #if STRIP_BLANK_LINE_WHITESPACE
@@ -1655,6 +1656,7 @@ static int copy_to_action_point_macro_expansion(parse_state* cs)
 	// expand output so copying entire file unmodified is possible
 	int out_len = arrlen(out);
 	arrsetcap(out, out_len + (cs->src_length - cs->src_offset) + SSE_READ_PADDING);
+	cs->dest = out;  // arrsetcap may realloc 'out' to a new address; keep cs->dest in sync so copy_and_filter_macro uses a valid pointer
 	q = out + out_len;
 
 	STB_ASSUME(q != NULL);
