[shaders] fix preprocessor tests to correctly initialize FShaderSource object to ensure sufficient padding for SSE reads

#Synced-CL 51762665

diff --git a/Engine/Source/Developer/ShaderPreprocessor/Tests/ShaderPreprocessorTests.cpp b/Engine/Source/Developer/ShaderPreprocessor/Tests/ShaderPreprocessorTests.cpp
--- a/Engine/Source/Developer/ShaderPreprocessor/Tests/ShaderPreprocessorTests.cpp
+++ b/Engine/Source/Developer/ShaderPreprocessor/Tests/ShaderPreprocessorTests.cpp
@@ -87,7 +87,7 @@ struct FPreprocessTestResult
 	}
 };
 
-FPreprocessTestResult ExecutePreprocessTest(const FShaderSource::FViewType& Source)
+FPreprocessTestResult ExecutePreprocessTest(const FShaderSource& Source)
 {
 	FPreprocessTestResult Result;
 	static bool bInitialized = false;
@@ -96,17 +96,18 @@ FPreprocessTestResult ExecutePreprocessTest(const FShaderSource::FViewType& Sour
 		init_preprocessor(&TestLoadFile, &TestFreeFile, nullptr, nullptr, nullptr);
 		bInitialized = true;
 	}
-	Result.Source = preprocess_file("preprocessor_test", const_cast<FShaderSource::FViewType*>(&Source), nullptr, 0u, &Result.Diagnostics, &Result.NumDiagnostics, nullptr, 0u);
+	FShaderSource::FViewType SourceView = Source.GetView();
+	Result.Source = preprocess_file("preprocessor_test", &SourceView, nullptr, 0u, &Result.Diagnostics, &Result.NumDiagnostics, nullptr, 0u);
 	return Result;
 }
 
 TEST_CASE("Shaders::ShaderPreprocessor::ConditionalExpressionErrorReporting", "[EditorContext][EngineFilter]")
 {
 	// need this dance of multiple assembled string literal fragments to avoid the C++ preprocessor parsing directives inside the literal
-	FShaderSource::FViewType Source = ANSITEXTVIEW(
+	FShaderSource Source(ANSITEXTVIEW(
 		"#" "if 0xFFFFFFFF+1 == 2\n"
 		"#" "error \"Unreachable\n\""
-		"#" "endif\n");
+		"#" "endif\n"));
 	FPreprocessTestResult Result = ExecutePreprocessTest(Source);
 	CHECK(Result.NumDiagnostics == 1);
 	CHECK(Result.Diagnostics[0].where->line_number == 1);
@@ -117,8 +118,8 @@ TEST_CASE("Shaders::ShaderPreprocessor::InvalidCharErrors", "[EditorContext][Eng
 {
 	{
 		// EOF after #define
-		FShaderSource::FViewType Source = ANSITEXTVIEW(
-			"#" "define");
+		FShaderSource Source(ANSITEXTVIEW(
+			"#" "define"));
 		FPreprocessTestResult Result = ExecutePreprocessTest(Source);
 		CHECK(Result.NumDiagnostics == 1);
 		CHECK(Result.Diagnostics[0].where->line_number == 1);
@@ -127,8 +128,8 @@ TEST_CASE("Shaders::ShaderPreprocessor::InvalidCharErrors", "[EditorContext][Eng
 
 	{
 		// Invalid printable char in define param list
-		FShaderSource::FViewType Source = ANSITEXTVIEW(
-			"#" "define BLAH(abc,#) abc\n");
+		FShaderSource Source(ANSITEXTVIEW(
+			"#" "define BLAH(abc,#) abc\n"));
 		FPreprocessTestResult Result = ExecutePreprocessTest(Source);
 		CHECK(Result.NumDiagnostics == 1);
 		CHECK(Result.Diagnostics[0].where->line_number == 1);
@@ -137,8 +138,8 @@ TEST_CASE("Shaders::ShaderPreprocessor::InvalidCharErrors", "[EditorContext][Eng
 
 	{
 		// Invalid non-printable char in define param list
-		FShaderSource::FViewType Source = ANSITEXTVIEW(
-			"#" "define BLAH(abc,\3) abc\n");
+		FShaderSource Source(ANSITEXTVIEW(
+			"#" "define BLAH(abc,\3) abc\n"));
 		FPreprocessTestResult Result = ExecutePreprocessTest(Source);
 		CHECK(Result.NumDiagnostics == 1);
 		CHECK(Result.Diagnostics[0].where->line_number == 1);
@@ -147,11 +148,11 @@ TEST_CASE("Shaders::ShaderPreprocessor::InvalidCharErrors", "[EditorContext][Eng
 
 	{
 		// Unexpected char after #if defined(
-		FShaderSource::FViewType Source = ANSITEXTVIEW(
+		FShaderSource Source(ANSITEXTVIEW(
 			"#" "define BLAH 1\n"
 			"#" "if defined(BLAH*)\n"
 			"#" "endif\n"
-		);
+		));
 		FPreprocessTestResult Result = ExecutePreprocessTest(Source);
 		CHECK(Result.NumDiagnostics == 1);
 		CHECK(Result.Diagnostics[0].where->line_number == 2);
@@ -161,11 +162,11 @@ TEST_CASE("Shaders::ShaderPreprocessor::InvalidCharErrors", "[EditorContext][Eng
 
 TEST_CASE("Shaders::ShaderPreprocessor::BuiltinMacros", "[EditorContext][EngineFilter]")
 {
-	FShaderSource::FViewType Source = ANSITEXTVIEW(
+	FShaderSource Source(ANSITEXTVIEW(
 		"int A = __" "LINE__;\n"
 		"int B = __" "COUNTER__;\n"
 		"int C = __" "COUNTER__;\n"
-	);
+	));
 	FPreprocessTestResult Result = ExecutePreprocessTest(Source);
 	CHECK(Result.NumDiagnostics == 0);
 	CHECK(strstr(Result.Source, "int A = 1") != nullptr);
@@ -177,7 +178,7 @@ TEST_CASE("Shaders::ShaderPreprocessor::EmitLineDirective", "[EditorContext][Eng
 {
 	// line directive will be emitted if 12 or more lines are removed due to an inactive preprocessor block
 	// (otherwise the inactive code will be replaced with empty lines)
-	FShaderSource::FViewType Source = ANSITEXTVIEW(
+	FShaderSource Source(ANSITEXTVIEW(
 		"#" "define MYDIR 0\n"
 		"#" "if MYDIR\n"
 		"int A = 0;\n"
@@ -197,7 +198,8 @@ TEST_CASE("Shaders::ShaderPreprocessor::EmitLineDirective", "[EditorContext][Eng
 		"B++;\n"
 		"#" "endif\n"
 		"int C = 0;\n"
-	);
+	));
+
 	// exactly two line directives will be emitted, one on the first line (done for all files) and one after the second inactive block (starting with the newline on line 18)
 	FPreprocessTestResult Result = ExecutePreprocessTest(Source);
 	CHECK(Result.NumDiagnostics == 0);
